REL - Change texture storage to Array64 to support 16K and larger Texture import

diff --git a/Engine/Plugins/2D/Paper2D/Source/Paper2D/Public/AlphaBitmap.h b/Engine/Plugins/2D/Paper2D/Source/Paper2D/Public/AlphaBitmap.h
--- a/Engine/Plugins/2D/Paper2D/Source/Paper2D/Public/AlphaBitmap.h
+++ b/Engine/Plugins/2D/Paper2D/Source/Paper2D/Public/AlphaBitmap.h
@@ -37,7 +37,7 @@ struct PAPER2D_API FAlphaBitmap
 
 		if (TextureSource != nullptr)
 		{
-			TArray<uint8> TextureRawData;
+			TArray64<uint8> TextureRawData;
 			TextureSource->GetMipData(TextureRawData, 0);
 			int32 BytesPerPixel = TextureSource->GetBytesPerPixel();
 			ETextureSourceFormat PixelFormat = TextureSource->GetFormat();
diff --git a/Engine/Plugins/Editor/BlueprintMaterialTextureNodes/Source/BlueprintMaterialTextureNodes/Private/BlueprintMaterialTextureNodesBPLibrary.cpp b/Engine/Plugins/Editor/BlueprintMaterialTextureNodes/Source/BlueprintMaterialTextureNodes/Private/BlueprintMaterialTextureNodesBPLibrary.cpp
--- a/Engine/Plugins/Editor/BlueprintMaterialTextureNodes/Source/BlueprintMaterialTextureNodes/Private/BlueprintMaterialTextureNodesBPLibrary.cpp
+++ b/Engine/Plugins/Editor/BlueprintMaterialTextureNodes/Source/BlueprintMaterialTextureNodes/Private/BlueprintMaterialTextureNodesBPLibrary.cpp
@@ -80,7 +80,7 @@ FLinearColor UBlueprintMaterialTextureNodesBPLibrary::Texture2D_SampleUV_EditorO
 		{
 			FTextureSource& TextureSource = Texture->Source;
 
-			TArray<uint8> SourceData;
+			TArray64<uint8> SourceData;
 			Texture->Source.GetMipData(SourceData, Mip);
 			ETextureSourceFormat SourceFormat = TextureSource.GetFormat();
 			int32 Index = ((Y * MipWidth) + X) * TextureSource.GetBytesPerPixel();
diff --git a/Engine/Plugins/Experimental/BlastPlugin/Source/BlastAuthoring/Private/FractureMesh.cpp b/Engine/Plugins/Experimental/BlastPlugin/Source/BlastAuthoring/Private/FractureMesh.cpp
--- a/Engine/Plugins/Experimental/BlastPlugin/Source/BlastAuthoring/Private/FractureMesh.cpp
+++ b/Engine/Plugins/Experimental/BlastPlugin/Source/BlastAuthoring/Private/FractureMesh.cpp
@@ -680,7 +680,7 @@ void UFractureMesh::ExtractDataFromTexture(const TWeakObjectPtr<UTexture> Source
 
 	if (TextureSource != nullptr)
 	{
-		TArray<uint8> TextureRawData;
+		TArray64<uint8> TextureRawData;
 		TextureSource->GetMipData(TextureRawData, 0);
 		int32 BytesPerPixel = TextureSource->GetBytesPerPixel();
 		ETextureSourceFormat PixelFormat = TextureSource->GetFormat();
diff --git a/Engine/Plugins/Runtime/AR/Google/GoogleARCore/Source/GoogleARCoreBase/Private/GoogleARCoreCookSupport.cpp b/Engine/Plugins/Runtime/AR/Google/GoogleARCore/Source/GoogleARCoreBase/Private/GoogleARCoreCookSupport.cpp
--- a/Engine/Plugins/Runtime/AR/Google/GoogleARCore/Source/GoogleARCoreBase/Private/GoogleARCoreCookSupport.cpp
+++ b/Engine/Plugins/Runtime/AR/Google/GoogleARCore/Source/GoogleARCoreBase/Private/GoogleARCoreCookSupport.cpp
@@ -29,7 +29,7 @@
 
 bool FGoogleARCoreSessionConfigCookSupport::SaveTextureToPNG(UTexture2D *Tex, const FString &Filename)
 {
-	TArray<uint8> MipData;
+	TArray64<uint8> MipData;
 	bool Ret = true;
 
 	if (Tex->Source.GetMipData(MipData, 0))
diff --git a/Engine/Source/Editor/ComponentVisualizers/Private/PointLightComponentVisualizer.cpp b/Engine/Source/Editor/ComponentVisualizers/Private/PointLightComponentVisualizer.cpp
--- a/Engine/Source/Editor/ComponentVisualizers/Private/PointLightComponentVisualizer.cpp
+++ b/Engine/Source/Editor/ComponentVisualizers/Private/PointLightComponentVisualizer.cpp
@@ -35,7 +35,7 @@ namespace TextureLightProfileVisualizerImpl
 {
 	struct FTextureLightProfileData
 	{
-		const TArray< uint8 >& Data;
+		const TArray64< uint8 >& Data;
 
 		int32 SizeX;
 		int32 SizeY;
@@ -267,7 +267,7 @@ void FTextureLightProfileVisualizer::UpdateIntensitiesCache(UTextureLightProfile
 
 	const FVector StartPos = LightTM.GetTranslation();
 
-	TArray< uint8 > MipData;
+	TArray64< uint8 > MipData;
 	TextureLightProfile->Source.GetMipData( MipData, 0 );
 
 	FTextureLightProfileData TextureLightProfileData{ MipData, TextureLightProfile->Source.GetSizeX(), TextureLightProfile->Source.GetSizeY(), TextureLightProfile->Source.GetBytesPerPixel() };
diff --git a/Engine/Source/Editor/LandscapeEditor/Private/LandscapeEditorObject.cpp b/Engine/Source/Editor/LandscapeEditor/Private/LandscapeEditorObject.cpp
--- a/Engine/Source/Editor/LandscapeEditor/Private/LandscapeEditorObject.cpp
+++ b/Engine/Source/Editor/LandscapeEditor/Private/LandscapeEditorObject.cpp
@@ -463,7 +463,7 @@ bool ULandscapeEditorObject::SetAlphaTexture(UTexture2D* InTexture, EColorChanne
 {
 	bool Result = true;
 
-	TArray<uint8> NewTextureData;
+	TArray64<uint8> NewTextureData;
 	UTexture2D* NewAlphaTexture = InTexture;
 
 	// No texture or no source art, try to use the previous texture.
diff --git a/Engine/Source/Editor/MeshPaint/Private/MeshPaintHelpers.cpp b/Engine/Source/Editor/MeshPaint/Private/MeshPaintHelpers.cpp
--- a/Engine/Source/Editor/MeshPaint/Private/MeshPaintHelpers.cpp
+++ b/Engine/Source/Editor/MeshPaint/Private/MeshPaintHelpers.cpp
@@ -1433,7 +1433,7 @@ void MeshPaintHelpers::ImportVertexColorsToStaticMesh(UStaticMesh* StaticMesh, c
 	checkf(StaticMesh && Options && Texture, TEXT("Invalid ptr"));
 
 	// Extract color data from texture
-	TArray<uint8> SrcMipData;
+	TArray64<uint8> SrcMipData;
 	Texture->Source.GetMipData(SrcMipData, 0);
 	const uint8* MipData = SrcMipData.GetData();
 
@@ -1478,7 +1478,7 @@ void MeshPaintHelpers::ImportVertexColorsToStaticMeshComponent(UStaticMeshCompon
 	checkf(StaticMeshComponent && Options && Texture, TEXT("Invalid ptr"));
 
 	// Extract color data from texture
-	TArray<uint8> SrcMipData;
+	TArray64<uint8> SrcMipData;
 	Texture->Source.GetMipData(SrcMipData, 0);
 	const uint8* MipData = SrcMipData.GetData();
 
@@ -1548,7 +1548,7 @@ void MeshPaintHelpers::ImportVertexColorsToSkeletalMesh(USkeletalMesh* SkeletalM
 	checkf(SkeletalMesh && Options && Texture, TEXT("Invalid ptr"));
 
 	// Extract color data from texture
-	TArray<uint8> SrcMipData;
+	TArray64<uint8> SrcMipData;
 	Texture->Source.GetMipData(SrcMipData, 0);
 	const uint8* MipData = SrcMipData.GetData();
 
diff --git a/Engine/Source/Editor/MeshPaint/Private/TexturePaintHelpers.cpp b/Engine/Source/Editor/MeshPaint/Private/TexturePaintHelpers.cpp
--- a/Engine/Source/Editor/MeshPaint/Private/TexturePaintHelpers.cpp
+++ b/Engine/Source/Editor/MeshPaint/Private/TexturePaintHelpers.cpp
@@ -326,7 +326,7 @@ UTexture2D* TexturePaintHelpers::CreateTempUncompressedTexture(UTexture2D* Sourc
 	check(SourceTexture->Source.IsValid());
 
 	// Decompress PNG image
-	TArray<uint8> RawData;
+	TArray64<uint8> RawData;
 	SourceTexture->Source.GetMipData(RawData, 0);
 
 	// We are using the source art so grab the original width/height
diff --git a/Engine/Source/Editor/UnrealEd/Private/EditorExporters.cpp b/Engine/Source/Editor/UnrealEd/Private/EditorExporters.cpp
--- a/Engine/Source/Editor/UnrealEd/Private/EditorExporters.cpp
+++ b/Engine/Source/Editor/UnrealEd/Private/EditorExporters.cpp
@@ -847,7 +847,7 @@ static void AddActorToOBJs(AActor* Actor, TArray<FOBJGeom*>& Objects, TSet<UMate
 			OBJGeom->Faces.AddZeroed( FMath::Square(ComponentSizeQuads) * 2 );
 
 			// Check if there are any holes
-			TArray<uint8> RawVisData;
+			TArray64<uint8> RawVisData;
 			uint8* VisDataMap = NULL;
 			int32 TexIndex = INDEX_NONE;
 			int32 WeightMapSize = (SubsectionSizeQuads + 1) * Component->NumSubsections;
diff --git a/Engine/Source/Editor/UnrealEd/Private/Factories/EditorFactories.cpp b/Engine/Source/Editor/UnrealEd/Private/Factories/EditorFactories.cpp
--- a/Engine/Source/Editor/UnrealEd/Private/Factories/EditorFactories.cpp
+++ b/Engine/Source/Editor/UnrealEd/Private/Factories/EditorFactories.cpp
@@ -403,7 +403,7 @@ UObject* UTexture2DFactoryNew::FactoryCreateNew( UClass* InClass, UObject* InPar
 	//Set the source art to be white as default.
 	if( Object->Source.IsValid() )
 	{
-		TArray<uint8> TexturePixels;
+		TArray64<uint8> TexturePixels;
 		Object->Source.GetMipData( TexturePixels, 0 );
 
 		uint8* DestData = Object->Source.LockMip(0);
@@ -4540,7 +4540,7 @@ bool UTextureExporterPCX::ExportBinary( UObject* Object, const TCHAR* Type, FArc
 
 	int32 SizeX = Texture->Source.GetSizeX();
 	int32 SizeY = Texture->Source.GetSizeY();
-	TArray<uint8> RawData;
+	TArray64<uint8> RawData;
 	Texture->Source.GetMipData(RawData, 0);
 
 	// Set all PCX file header properties.
@@ -4635,7 +4635,7 @@ bool UTextureExporterBMP::ExportBinary( UObject* Object, const TCHAR* Type, FArc
 
 	int32 SizeX = Texture->Source.GetSizeX();
 	int32 SizeY = Texture->Source.GetSizeY();
-	TArray<uint8> RawData;
+	TArray64<uint8> RawData;
 	Texture->Source.GetMipData(RawData, 0);
 
 	FBitmapFileHeader bmf;
@@ -4835,7 +4835,7 @@ bool UTextureExporterTGA::ExportBinary( UObject* Object, const TCHAR* Type, FArc
 
 	int32 SizeX = Texture->Source.GetSizeX();
 	int32 SizeY = Texture->Source.GetSizeY();
-	TArray<uint8> RawData;
+	TArray64<uint8> RawData;
 	Texture->Source.GetMipData(RawData, 0);
 
 	// If we should export the file with no alpha info.  
diff --git a/Engine/Source/Editor/UnrealEd/Private/Lightmass/Lightmass.cpp b/Engine/Source/Editor/UnrealEd/Private/Lightmass/Lightmass.cpp
--- a/Engine/Source/Editor/UnrealEd/Private/Lightmass/Lightmass.cpp
+++ b/Engine/Source/Editor/UnrealEd/Private/Lightmass/Lightmass.cpp
@@ -241,7 +241,7 @@ void CopyLightProfile( const ULightComponent* In, Lightmass::FLightData& Out, TA
 		{
 			Out.LightFlags |= Lightmass::GI_LIGHT_USE_LIGHTPROFILE;
 
-			TArray<uint8> MipData;
+			TArray64<uint8> MipData;
 
 			Source.GetMipData(MipData, 0);
 
@@ -1174,7 +1174,7 @@ void FLightmassExporter::WriteLights( int32 Channel )
 				{
 					//int32 i = SourceTexture.AddUninitialized( SizeX * SizeY );
 					
-					TArray< uint8 > MipData;
+					TArray64< uint8 > MipData;
 					Source.GetMipData( MipData, MipLevel );
 
 					uint8* Pixel = MipData.GetData();
diff --git a/Engine/Source/Runtime/Engine/Classes/Engine/Texture.h b/Engine/Source/Runtime/Engine/Classes/Engine/Texture.h
--- a/Engine/Source/Runtime/Engine/Classes/Engine/Texture.h
+++ b/Engine/Source/Runtime/Engine/Classes/Engine/Texture.h
@@ -263,7 +263,7 @@ struct FTextureSource
 	ENGINE_API void UnlockMip(int32 BlockIndex, int32 LayerIndex, int32 MipIndex);
 
 	/** Retrieve a copy of the data for a particular mip. */
-	ENGINE_API bool GetMipData(TArray<uint8>& OutMipData, int32 BlockIndex, int32 LayerIndex, int32 MipIndex, class IImageWrapperModule* ImageWrapperModule = nullptr);
+	ENGINE_API bool GetMipData(TArray64<uint8>& OutMipData, int32 BlockIndex, int32 LayerIndex, int32 MipIndex, class IImageWrapperModule* ImageWrapperModule = nullptr);
 
 	/** Computes the size of a single mip. */
 	ENGINE_API int32 CalcMipSize(int32 BlockIndex, int32 LayerIndex, int32 MipIndex) const;
@@ -308,7 +308,7 @@ struct FTextureSource
 	ENGINE_API void SetId(const FGuid& InId, bool bInGuidIsHash);
 
 	/** Legacy API that defaults to LayerIndex 0 */
-	FORCEINLINE bool GetMipData(TArray<uint8>& OutMipData, int32 MipIndex, class IImageWrapperModule* ImageWrapperModule = nullptr)
+	FORCEINLINE bool GetMipData(TArray64<uint8>& OutMipData, int32 MipIndex, class IImageWrapperModule* ImageWrapperModule = nullptr)
 	{
 		return GetMipData(OutMipData, 0, 0, MipIndex, ImageWrapperModule);
 	}
diff --git a/Engine/Source/Runtime/Engine/Private/Atmosphere/AtmosphericFogComponent.cpp b/Engine/Source/Runtime/Engine/Private/Atmosphere/AtmosphericFogComponent.cpp
--- a/Engine/Source/Runtime/Engine/Private/Atmosphere/AtmosphericFogComponent.cpp
+++ b/Engine/Source/Runtime/Engine/Private/Atmosphere/AtmosphericFogComponent.cpp
@@ -127,7 +127,7 @@ void UAtmosphericFogComponent::PostLoad()
 			// Copy data from Previous data
 			if (TransmittanceTexture_DEPRECATED->Source.IsValid())
 			{
-				TArray<uint8> RawData;
+				TArray64<uint8> RawData;
 				TArray<FColor> OutData;
 				TransmittanceTexture_DEPRECATED->Source.GetMipData(RawData, 0);
 				// Convert from FFloat16Color to FColor
@@ -155,7 +155,7 @@ void UAtmosphericFogComponent::PostLoad()
 		{
 			if (IrradianceTexture_DEPRECATED->Source.IsValid())
 			{
-				TArray<uint8> RawData;
+				TArray64<uint8> RawData;
 				TArray<FColor> OutData;
 				IrradianceTexture_DEPRECATED->Source.GetMipData(RawData, 0);
 				// Convert from FFloat16Color to FColor
diff --git a/Engine/Source/Runtime/Engine/Private/CubemapUnwrapUtils.cpp b/Engine/Source/Runtime/Engine/Private/CubemapUnwrapUtils.cpp
--- a/Engine/Source/Runtime/Engine/Private/CubemapUnwrapUtils.cpp
+++ b/Engine/Source/Runtime/Engine/Private/CubemapUnwrapUtils.cpp
@@ -34,7 +34,7 @@ namespace CubemapHelpers
 	* @return						true on success.
 	* @param	FormatOUT			Filled with the pixel format of the output bitmap.
 	*/
-	bool GenerateLongLatUnwrap(const FTextureResource* TextureResource, const uint32 AxisDimenion, const EPixelFormat SourcePixelFormat, TArray<uint8>& BitsOUT, FIntPoint& SizeOUT, EPixelFormat& FormatOUT)
+	bool GenerateLongLatUnwrap(const FTextureResource* TextureResource, const uint32 AxisDimenion, const EPixelFormat SourcePixelFormat, TArray64<uint8>& BitsOUT, FIntPoint& SizeOUT, EPixelFormat& FormatOUT)
 	{
 		TRefCountPtr<FBatchedElementParameters> BatchedElementParameters;
 		BatchedElementParameters = new FMipLevelBatchedElementParameters((float)0, true);
@@ -102,13 +102,13 @@ namespace CubemapHelpers
 		return bReadSuccess;
 	}
 
-	bool GenerateLongLatUnwrap(const UTextureCube* CubeTexture, TArray<uint8>& BitsOUT, FIntPoint& SizeOUT, EPixelFormat& FormatOUT)
+	bool GenerateLongLatUnwrap(const UTextureCube* CubeTexture, TArray64<uint8>& BitsOUT, FIntPoint& SizeOUT, EPixelFormat& FormatOUT)
 	{
 		check(CubeTexture != NULL);
 		return GenerateLongLatUnwrap(CubeTexture->Resource, CubeTexture->GetSizeX(), CubeTexture->GetPixelFormat(), BitsOUT, SizeOUT, FormatOUT);
 	}
 
-	bool GenerateLongLatUnwrap(const UTextureRenderTargetCube* CubeTarget, TArray<uint8>& BitsOUT, FIntPoint& SizeOUT, EPixelFormat& FormatOUT)
+	bool GenerateLongLatUnwrap(const UTextureRenderTargetCube* CubeTarget, TArray64<uint8>& BitsOUT, FIntPoint& SizeOUT, EPixelFormat& FormatOUT)
 	{
 		check(CubeTarget != NULL);
 		return GenerateLongLatUnwrap(CubeTarget->Resource, CubeTarget->SizeX, CubeTarget->GetFormat(), BitsOUT, SizeOUT, FormatOUT);
diff --git a/Engine/Source/Runtime/Engine/Private/ImageUtils.cpp b/Engine/Source/Runtime/Engine/Private/ImageUtils.cpp
--- a/Engine/Source/Runtime/Engine/Private/ImageUtils.cpp
+++ b/Engine/Source/Runtime/Engine/Private/ImageUtils.cpp
@@ -22,7 +22,7 @@ DEFINE_LOG_CATEGORY_STATIC(LogImageUtils, Log, All);
 
 #define LOCTEXT_NAMESPACE "ImageUtils"
 
-static bool GetRawData(UTextureRenderTarget2D* TexRT, TArray<uint8>& RawData)
+static bool GetRawData(UTextureRenderTarget2D* TexRT, TArray64<uint8>& RawData)
 {
 	FRenderTarget* RenderTarget = TexRT->GameThread_GetRenderTargetResource();
 	EPixelFormat Format = TexRT->GetFormat();
@@ -375,7 +375,7 @@ public:
 		Size = RenderTarget->GetSizeXY();
 		Format = TexRT->GetFormat();
 
-		TArray<uint8> RawData;
+		TArray64<uint8> RawData;
 		bool bReadSuccess = GetRawData(TexRT, RawData);
 		if (bReadSuccess)
 		{
@@ -395,7 +395,7 @@ public:
 	{
 		check(Texture != nullptr);
 		bool bReadSuccess = true;
-		TArray<uint8> RawData;
+		TArray64<uint8> RawData;
 
 #if WITH_EDITORONLY_DATA
 		Size = FIntPoint(Texture->Source.GetSizeX(), Texture->Source.GetSizeY());
@@ -479,7 +479,7 @@ public:
 		check(TexCube != nullptr);
 
 		// Generate 2D image.
-		TArray<uint8> RawData;
+		TArray64<uint8> RawData;
 		bool bUnwrapSuccess = CubemapHelpers::GenerateLongLatUnwrap(TexCube, RawData, Size, Format);
 		bool bAcceptableFormat = (Format == PF_B8G8R8A8 || Format == PF_FloatRGBA);
 		if (bUnwrapSuccess == false || bAcceptableFormat == false)
@@ -504,7 +504,7 @@ public:
 		check(TexCube != nullptr);
 
 		// Generate 2D image.
-		TArray<uint8> RawData;
+		TArray64<uint8> RawData;
 		bool bUnwrapSuccess = CubemapHelpers::GenerateLongLatUnwrap(TexCube, RawData, Size, Format);
 		bool bAcceptableFormat = (Format == PF_B8G8R8A8 || Format == PF_FloatRGBA);
 		if (bUnwrapSuccess == false || bAcceptableFormat == false)
@@ -626,7 +626,7 @@ private:
 		Ar.Serialize(Header, Len);
 	}
 
-	void WriteHDRImage(const TArray<uint8>& RawData, FArchive& Ar)
+	void WriteHDRImage(const TArray64<uint8>& RawData, FArchive& Ar)
 	{
 		WriteHDRHeader(Ar);
 		if (Format == PF_FloatRGBA)
@@ -692,7 +692,7 @@ bool FImageUtils::ExportRenderTarget2DAsPNG(UTextureRenderTarget2D* TexRT, FArch
 		FRenderTarget* RenderTarget = TexRT->GameThread_GetRenderTargetResource();
 		FIntPoint Size = RenderTarget->GetSizeXY();
 
-		TArray<uint8> RawData;
+		TArray64<uint8> RawData;
 		bSuccess = GetRawData(TexRT, RawData);
 
 		IImageWrapperModule& ImageWrapperModule = FModuleManager::Get().LoadModuleChecked<IImageWrapperModule>(TEXT("ImageWrapper"));
@@ -718,7 +718,7 @@ ENGINE_API bool FImageUtils::ExportRenderTarget2DAsEXR(UTextureRenderTarget2D* T
 		FRenderTarget* RenderTarget = TexRT->GameThread_GetRenderTargetResource();
 		FIntPoint Size = RenderTarget->GetSizeXY();
 
-		TArray<uint8> RawData;
+		TArray64<uint8> RawData;
 		bSuccess = GetRawData(TexRT, RawData);
 
 		int32 BitsPerPixel = TexRT->GetFormat() == PF_B8G8R8A8 ? 8 : (sizeof(FFloat16Color) / 4) * 8;
diff --git a/Engine/Source/Runtime/Engine/Private/Particles/SubUVAnimation.cpp b/Engine/Source/Runtime/Engine/Private/Particles/SubUVAnimation.cpp
--- a/Engine/Source/Runtime/Engine/Private/Particles/SubUVAnimation.cpp
+++ b/Engine/Source/Runtime/Engine/Private/Particles/SubUVAnimation.cpp
@@ -479,7 +479,7 @@ uint8 ComputeOpacityValue(const uint8* BGRA, EOpacitySourceMode OpacitySourceMod
 }
 
 /** Counts how many neighbors have non-zero alpha. */
-int32 ComputeNeighborCount(int32 X, int32 Y, int32 GlobalX, int32 GlobalY, int32 SubImageSizeX, int32 SubImageSizeY, int32 TextureSizeX, const TArray<uint8>& MipData, uint8 AlphaThresholdByte, EOpacitySourceMode OpacitySourceMode)
+int32 ComputeNeighborCount(int32 X, int32 Y, int32 GlobalX, int32 GlobalY, int32 SubImageSizeX, int32 SubImageSizeY, int32 TextureSizeX, const TArray64<uint8>& MipData, uint8 AlphaThresholdByte, EOpacitySourceMode OpacitySourceMode)
 {
 	int32 NeighborCount = 0;
 
@@ -535,7 +535,7 @@ void FSubUVDerivedData::Build(UTexture2D* SubUVTexture, int32 SubImages_Horizont
 
 	if (SubUVTexture)
 	{
-		TArray<uint8> MipData;
+		TArray64<uint8> MipData;
 		bool bSuccess = SubUVTexture->Source.GetFormat() == TSF_BGRA8 && SubUVTexture->Source.GetMipData(MipData, 0);
 
 		const int32 TextureSizeX = SubUVTexture->Source.GetSizeX();
diff --git a/Engine/Source/Runtime/Engine/Private/Texture.cpp b/Engine/Source/Runtime/Engine/Private/Texture.cpp
--- a/Engine/Source/Runtime/Engine/Private/Texture.cpp
+++ b/Engine/Source/Runtime/Engine/Private/Texture.cpp
@@ -917,7 +917,7 @@ void FTextureSource::UnlockMip(int32 BlockIndex, int32 LayerIndex, int32 MipInde
 	}
 }
 
-bool FTextureSource::GetMipData(TArray<uint8>& OutMipData, int32 BlockIndex, int32 LayerIndex, int32 MipIndex, IImageWrapperModule* ImageWrapperModule)
+bool FTextureSource::GetMipData(TArray64<uint8>& OutMipData, int32 BlockIndex, int32 LayerIndex, int32 MipIndex, IImageWrapperModule* ImageWrapperModule)
 {
 	bool bSuccess = false;
 	if (BlockIndex < GetNumBlocks() && LayerIndex < NumLayers && MipIndex < NumMips && BulkData.GetBulkDataSize() > 0)
diff --git a/Engine/Source/Runtime/Engine/Private/Texture2D.cpp b/Engine/Source/Runtime/Engine/Private/Texture2D.cpp
--- a/Engine/Source/Runtime/Engine/Private/Texture2D.cpp
+++ b/Engine/Source/Runtime/Engine/Private/Texture2D.cpp
@@ -580,7 +580,7 @@ float UTexture2D::GetAverageBrightness(bool bIgnoreTrueBlack, bool bUseGrayscale
 {
 	float AvgBrightness = -1.0f;
 #if WITH_EDITOR
-	TArray<uint8> RawData;
+	TArray64<uint8> RawData;
 	// use the source art if it exists
 	if (Source.IsValid() && Source.GetFormat() == TSF_BGRA8)
 	{
@@ -976,7 +976,7 @@ uint32 UTexture2D::CalcTextureMemorySizeEnum( ETextureMipCount Enum ) const
 bool UTexture2D::GetSourceArtCRC(uint32& OutSourceCRC)
 {
 	bool bResult = false;
-	TArray<uint8> RawData;
+	TArray64<uint8> RawData;
 #if WITH_EDITOR
 	// use the source art if it exists
 	if (Source.IsValid())
@@ -1002,8 +1002,8 @@ bool UTexture2D::HasSameSourceArt(UTexture2D* InTexture)
 {
 	bool bResult = false;
 #if WITH_EDITOR
-	TArray<uint8> RawData1;
-	TArray<uint8> RawData2;
+	TArray64<uint8> RawData1;
+	TArray64<uint8> RawData2;
 	int32 SizeX = 0;
 	int32 SizeY = 0;
 
diff --git a/Engine/Source/Runtime/Engine/Private/VolumeTexture.cpp b/Engine/Source/Runtime/Engine/Private/VolumeTexture.cpp
--- a/Engine/Source/Runtime/Engine/Private/VolumeTexture.cpp
+++ b/Engine/Source/Runtime/Engine/Private/VolumeTexture.cpp
@@ -46,7 +46,7 @@ bool UVolumeTexture::UpdateSourceFromSourceTexture()
 			const int32 FormatDataSize = InitialSource.GetBytesPerPixel();
 			if (FormatDataSize > 0)
 			{
-				TArray<uint8> Ref2DData;
+				TArray64<uint8> Ref2DData;
 				if (InitialSource.GetMipData(Ref2DData, 0))
 				{
 					uint8* NewData = (uint8*)FMemory::Malloc(Source2DTileSizeX * Source2DTileSizeY * TileSizeZ * FormatDataSize);
diff --git a/Engine/Source/Runtime/Engine/Public/CubemapUnwrapUtils.h b/Engine/Source/Runtime/Engine/Public/CubemapUnwrapUtils.h
--- a/Engine/Source/Runtime/Engine/Public/CubemapUnwrapUtils.h
+++ b/Engine/Source/Runtime/Engine/Public/CubemapUnwrapUtils.h
@@ -26,7 +26,7 @@ namespace CubemapHelpers
 	* @param	FormatOUT	Filled with the pixel format of the output bitmap.
 	* @return	true on success.
 	*/
-	ENGINE_API bool GenerateLongLatUnwrap(const UTextureCube* CubeTexture, TArray<uint8>& BitsOUT, FIntPoint& SizeOUT, EPixelFormat& FormatOUT);
+	ENGINE_API bool GenerateLongLatUnwrap(const UTextureCube* CubeTexture, TArray64<uint8>& BitsOUT, FIntPoint& SizeOUT, EPixelFormat& FormatOUT);
 
 	/**
 	* Creates an unwrapped 2D image of the cube map ( longitude/latitude )
@@ -37,7 +37,7 @@ namespace CubemapHelpers
 	* @param	FormatOUT	Filled with the pixel format of the output bitmap.
 	* @return	true on success.
 	*/
-	ENGINE_API bool GenerateLongLatUnwrap(const UTextureRenderTargetCube* CubeTarget, TArray<uint8>& BitsOUT, FIntPoint& SizeOUT, EPixelFormat& FormatOUT);
+	ENGINE_API bool GenerateLongLatUnwrap(const UTextureRenderTargetCube* CubeTarget, TArray64<uint8>& BitsOUT, FIntPoint& SizeOUT, EPixelFormat& FormatOUT);
 }
 
 /**
diff --git a/Engine/Source/Runtime/ImageCore/Private/ImageCore.cpp b/Engine/Source/Runtime/ImageCore/Private/ImageCore.cpp
--- a/Engine/Source/Runtime/ImageCore/Private/ImageCore.cpp
+++ b/Engine/Source/Runtime/ImageCore/Private/ImageCore.cpp
@@ -18,7 +18,7 @@ IMPLEMENT_MODULE(FDefaultModuleImpl, ImageCore);
  */
 static void InitImageStorage(FImage& Image)
 {
-	int32 NumBytes = Image.SizeX * Image.SizeY * Image.NumSlices * Image.GetBytesPerPixel();
+	int64 NumBytes = int64(Image.SizeX) * Image.SizeY * Image.NumSlices * Image.GetBytesPerPixel();
 	Image.RawData.Empty(NumBytes);
 	Image.RawData.AddUninitialized(NumBytes);
 }
@@ -37,9 +37,9 @@ static void CopyImage(const FImage& SrcImage, FImage& DestImage)
 	check(SrcImage.NumSlices == DestImage.NumSlices);
 
 	const bool bDestIsGammaCorrected = DestImage.IsGammaCorrected();
-	const int32 NumTexels = SrcImage.SizeX * SrcImage.SizeY * SrcImage.NumSlices;
-	const int32 NumJobs = FTaskGraphInterface::Get().GetNumWorkerThreads();
-	int32 TexelsPerJob = NumTexels / NumJobs;
+	const int64 NumTexels = int64(SrcImage.SizeX) * SrcImage.SizeY * SrcImage.NumSlices;
+	const int64 NumJobs = FTaskGraphInterface::Get().GetNumWorkerThreads();
+	int64 TexelsPerJob = NumTexels / NumJobs;
 	if (TexelsPerJob * NumJobs < NumTexels)
 	{
 		++TexelsPerJob;
@@ -60,11 +60,11 @@ static void CopyImage(const FImage& SrcImage, FImage& DestImage)
 		case ERawImageFormat::G8:
 			{
 				uint8* DestLum = DestImage.AsG8();
-				ParallelFor(NumJobs, [NumJobs, DestLum, SrcColors, TexelsPerJob, NumTexels, bDestIsGammaCorrected](int32 JobIndex)
+				ParallelFor(NumJobs, [NumJobs, DestLum, SrcColors, TexelsPerJob, NumTexels, bDestIsGammaCorrected](int64 JobIndex)
 				{
-					const int32 StartIndex = JobIndex * TexelsPerJob;
-					const int32 EndIndex = FMath::Min(StartIndex + TexelsPerJob, NumTexels);
-					for (int32 TexelIndex = StartIndex; TexelIndex < EndIndex; ++TexelIndex)
+					const int64 StartIndex = JobIndex * TexelsPerJob;
+					const int64 EndIndex = FMath::Min(StartIndex + TexelsPerJob, NumTexels);
+					for (int64 TexelIndex = StartIndex; TexelIndex < EndIndex; ++TexelIndex)
 					{
 						DestLum[TexelIndex] = SrcColors[TexelIndex].ToFColor(bDestIsGammaCorrected).R;
 					}
@@ -75,11 +75,11 @@ static void CopyImage(const FImage& SrcImage, FImage& DestImage)
 		case ERawImageFormat::BGRA8:
 			{
 				FColor* DestColors = DestImage.AsBGRA8();
-				ParallelFor(NumJobs, [DestColors, SrcColors, TexelsPerJob, NumTexels, bDestIsGammaCorrected](int32 JobIndex)
+				ParallelFor(NumJobs, [DestColors, SrcColors, TexelsPerJob, NumTexels, bDestIsGammaCorrected](int64 JobIndex)
 				{
-					const int32 StartIndex = JobIndex * TexelsPerJob;
-					const int32 EndIndex = FMath::Min(StartIndex + TexelsPerJob, NumTexels);
-					for (int32 TexelIndex = StartIndex; TexelIndex < EndIndex; ++TexelIndex)
+					const int64 StartIndex = JobIndex * TexelsPerJob;
+					const int64 EndIndex = FMath::Min(StartIndex + TexelsPerJob, NumTexels);
+					for (int64 TexelIndex = StartIndex; TexelIndex < EndIndex; ++TexelIndex)
 					{
 						DestColors[TexelIndex] = SrcColors[TexelIndex].ToFColor(bDestIsGammaCorrected);
 					}
@@ -90,11 +90,11 @@ static void CopyImage(const FImage& SrcImage, FImage& DestImage)
 		case ERawImageFormat::BGRE8:
 			{
 				FColor* DestColors = DestImage.AsBGRE8();
-				ParallelFor(NumJobs, [DestColors, SrcColors, TexelsPerJob, NumTexels, bDestIsGammaCorrected](int32 JobIndex)
+				ParallelFor(NumJobs, [DestColors, SrcColors, TexelsPerJob, NumTexels, bDestIsGammaCorrected](int64 JobIndex)
 				{
-					const int32 StartIndex = JobIndex * TexelsPerJob;
-					const int32 EndIndex = FMath::Min(StartIndex + TexelsPerJob, NumTexels);
-					for (int32 TexelIndex = StartIndex; TexelIndex < EndIndex; ++TexelIndex)
+					const int64 StartIndex = JobIndex * TexelsPerJob;
+					const int64 EndIndex = FMath::Min(StartIndex + TexelsPerJob, NumTexels);
+					for (int64 TexelIndex = StartIndex; TexelIndex < EndIndex; ++TexelIndex)
 					{
 						DestColors[TexelIndex] = SrcColors[TexelIndex].ToRGBE();
 					}
@@ -105,9 +105,9 @@ static void CopyImage(const FImage& SrcImage, FImage& DestImage)
 		case ERawImageFormat::RGBA16:
 			{
 				uint16* DestColors = DestImage.AsRGBA16();
-				for (int32 TexelIndex = 0; TexelIndex < NumTexels; ++TexelIndex)
+				for (int64 TexelIndex = 0; TexelIndex < NumTexels; ++TexelIndex)
 				{
-					int32 DestIndex = TexelIndex * 4;
+					int64 DestIndex = TexelIndex * 4;
 					DestColors[DestIndex + 0] = FMath::Clamp(FMath::FloorToInt(SrcColors[TexelIndex].R * 65535.999f), 0, 65535);
 					DestColors[DestIndex + 1] = FMath::Clamp(FMath::FloorToInt(SrcColors[TexelIndex].G * 65535.999f), 0, 65535);
 					DestColors[DestIndex + 2] = FMath::Clamp(FMath::FloorToInt(SrcColors[TexelIndex].B * 65535.999f), 0, 65535);
@@ -119,7 +119,7 @@ static void CopyImage(const FImage& SrcImage, FImage& DestImage)
 		case ERawImageFormat::RGBA16F:
 			{
 				FFloat16Color* DestColors = DestImage.AsRGBA16F();
-				for (int32 TexelIndex = 0; TexelIndex < NumTexels; ++TexelIndex)
+				for (int64 TexelIndex = 0; TexelIndex < NumTexels; ++TexelIndex)
 				{
 					DestColors[TexelIndex] = FFloat16Color(SrcColors[TexelIndex]);
 				}
@@ -136,7 +136,7 @@ static void CopyImage(const FImage& SrcImage, FImage& DestImage)
 		case ERawImageFormat::G8:
 			{
 				const uint8* SrcLum = SrcImage.AsG8();
-				for (int32 TexelIndex = 0; TexelIndex < NumTexels; ++TexelIndex)
+				for (int64 TexelIndex = 0; TexelIndex < NumTexels; ++TexelIndex)
 				{
 					FColor SrcColor(SrcLum[TexelIndex],SrcLum[TexelIndex],SrcLum[TexelIndex],255);
 					
@@ -162,29 +162,29 @@ static void CopyImage(const FImage& SrcImage, FImage& DestImage)
 				switch ( SrcImage.GammaSpace )
 				{
 				case EGammaSpace::Linear:
-					ParallelFor(NumJobs, [DestColors, SrcColors, TexelsPerJob, NumTexels](int32 JobIndex)
+					ParallelFor(NumJobs, [DestColors, SrcColors, TexelsPerJob, NumTexels](int64 JobIndex)
 					{
-						int32 StartIndex = JobIndex * TexelsPerJob;
-						int32 EndIndex = FMath::Min(StartIndex + TexelsPerJob, NumTexels);
-						for (int32 TexelIndex = StartIndex; TexelIndex < EndIndex; ++TexelIndex)
+						int64 StartIndex = JobIndex * TexelsPerJob;
+						int64 EndIndex = FMath::Min(StartIndex + TexelsPerJob, NumTexels);
+						for (int64 TexelIndex = StartIndex; TexelIndex < EndIndex; ++TexelIndex)
 						{
 							DestColors[TexelIndex] = SrcColors[TexelIndex].ReinterpretAsLinear();
 						}
 					});
 					break;
 				case EGammaSpace::sRGB:
-					ParallelFor(NumJobs, [DestColors, SrcColors, TexelsPerJob, NumTexels](int32 JobIndex)
+					ParallelFor(NumJobs, [DestColors, SrcColors, TexelsPerJob, NumTexels](int64 JobIndex)
 					{
-						int32 StartIndex = JobIndex * TexelsPerJob;
-						int32 EndIndex = FMath::Min(StartIndex + TexelsPerJob, NumTexels);
-						for (int32 TexelIndex = StartIndex; TexelIndex < EndIndex; ++TexelIndex)
+						int64 StartIndex = JobIndex * TexelsPerJob;
+						int64 EndIndex = FMath::Min(StartIndex + TexelsPerJob, NumTexels);
+						for (int64 TexelIndex = StartIndex; TexelIndex < EndIndex; ++TexelIndex)
 						{
 							DestColors[TexelIndex] = FLinearColor(SrcColors[TexelIndex]);
 						}
 					});
 					break;
 				case EGammaSpace::Pow22:
-					for (int32 TexelIndex = 0; TexelIndex < NumTexels; ++TexelIndex)
+					for (int64 TexelIndex = 0; TexelIndex < NumTexels; ++TexelIndex)
 					{
 						DestColors[TexelIndex] = FLinearColor::FromPow22Color(SrcColors[TexelIndex]);
 					}
@@ -196,7 +196,7 @@ static void CopyImage(const FImage& SrcImage, FImage& DestImage)
 		case ERawImageFormat::BGRE8:
 			{
 				const FColor* SrcColors = SrcImage.AsBGRE8();
-				for (int32 TexelIndex = 0; TexelIndex < NumTexels; ++TexelIndex)
+				for (int64 TexelIndex = 0; TexelIndex < NumTexels; ++TexelIndex)
 				{
 					DestColors[TexelIndex] = SrcColors[TexelIndex].FromRGBE();
 				}
@@ -206,9 +206,9 @@ static void CopyImage(const FImage& SrcImage, FImage& DestImage)
 		case ERawImageFormat::RGBA16:
 			{
 				const uint16* SrcColors = SrcImage.AsRGBA16();
-				for (int32 TexelIndex = 0; TexelIndex < NumTexels; ++TexelIndex)
+				for (int64 TexelIndex = 0; TexelIndex < NumTexels; ++TexelIndex)
 				{
-					int32 SrcIndex = TexelIndex * 4;
+					int64 SrcIndex = TexelIndex * 4;
 					DestColors[TexelIndex] = FLinearColor(
 						SrcColors[SrcIndex + 0] / 65535.0f,
 						SrcColors[SrcIndex + 1] / 65535.0f,
@@ -222,7 +222,7 @@ static void CopyImage(const FImage& SrcImage, FImage& DestImage)
 		case ERawImageFormat::RGBA16F:
 			{
 				const FFloat16Color* SrcColors = SrcImage.AsRGBA16F();
-				for (int32 TexelIndex = 0; TexelIndex < NumTexels; ++TexelIndex)
+				for (int64 TexelIndex = 0; TexelIndex < NumTexels; ++TexelIndex)
 				{
 					DestColors[TexelIndex] = FLinearColor(SrcColors[TexelIndex]);
 				}
@@ -241,10 +241,10 @@ static void CopyImage(const FImage& SrcImage, FImage& DestImage)
 
 static FLinearColor SampleImage(const FLinearColor* Pixels, int Width, int Height, float X, float Y)
 {
-	const int32 TexelX0 = FMath::FloorToInt(X);
-	const int32 TexelY0 = FMath::FloorToInt(Y);
-	const int32 TexelX1 = FMath::Min(TexelX0 + 1, Width - 1);
-	const int32 TexelY1 = FMath::Min(TexelY0 + 1, Height - 1);
+	const int64 TexelX0 = FMath::FloorToInt(X);
+	const int64 TexelY0 = FMath::FloorToInt(Y);
+	const int64 TexelX1 = FMath::Min<int64>(TexelX0 + 1, Width - 1);
+	const int64 TexelY1 = FMath::Min<int64>(TexelY0 + 1, Height - 1);
 	checkSlow(TexelX0 >= 0 && TexelX0 < Width);
 	checkSlow(TexelY0 >= 0 && TexelY0 < Width);
 
diff --git a/Engine/Source/Runtime/ImageCore/Public/ImageCore.h b/Engine/Source/Runtime/ImageCore/Public/ImageCore.h
--- a/Engine/Source/Runtime/ImageCore/Public/ImageCore.h
+++ b/Engine/Source/Runtime/ImageCore/Public/ImageCore.h
@@ -34,7 +34,7 @@ namespace ERawImageFormat
 struct FImage
 {
 	/** Raw image data. */
-	TArray<uint8> RawData;
+	TArray64<uint8> RawData;
 
 	/** Width of the image. */
 	int32 SizeX;
diff --git a/Engine/Source/Runtime/Landscape/Private/LandscapeEdit.cpp b/Engine/Source/Runtime/Landscape/Private/LandscapeEdit.cpp
--- a/Engine/Source/Runtime/Landscape/Private/LandscapeEdit.cpp
+++ b/Engine/Source/Runtime/Landscape/Private/LandscapeEdit.cpp
@@ -1025,9 +1025,9 @@ void ULandscapeComponent::DestroyCollisionData()
 
 void ULandscapeComponent::UpdateCollisionData(bool bInUpdateHeightfieldRegion)
 {
-	TArray<uint8> CollisionMipData;
-	TArray<uint8> SimpleCollisionMipData;
-	TArray<uint8> XYOffsetMipData;
+	TArray64<uint8> CollisionMipData;
+	TArray64<uint8> SimpleCollisionMipData;
+	TArray64<uint8> XYOffsetMipData;
 
 	GetHeightmap()->Source.GetMipData(CollisionMipData, CollisionMipLevel);
 	if (SimpleCollisionMipLevel > CollisionMipLevel)
@@ -1428,7 +1428,7 @@ void ULandscapeComponent::UpdateCollisionLayerData()
 	WeightmapTextureMipData.Empty(ComponentWeightmapsTexture.Num());
 	for (int32 WeightmapIdx = 0; WeightmapIdx < ComponentWeightmapsTexture.Num(); ++WeightmapIdx)
 	{
-		TArray<uint8> MipData;
+		TArray64<uint8> MipData;
 		ComponentWeightmapsTexture[WeightmapIdx]->Source.GetMipData(MipData, CollisionMipLevel);
 		WeightmapTextureMipData.Add((FColor*)MipData.GetData());
 	}
@@ -1438,7 +1438,7 @@ void ULandscapeComponent::UpdateCollisionLayerData()
 	{
 		for (int32 WeightmapIdx = 0; WeightmapIdx < ComponentWeightmapsTexture.Num(); ++WeightmapIdx)
 		{
-			TArray<uint8> MipData;
+			TArray64<uint8> MipData;
 			ComponentWeightmapsTexture[WeightmapIdx]->Source.GetMipData(MipData, SimpleCollisionMipLevel);
 			SimpleCollisionWeightmapMipData.Add((FColor*)MipData.GetData());
 		}
@@ -1711,7 +1711,7 @@ void ULandscapeComponent::CreateEmptyTextureMips(UTexture2D* Texture, bool bClea
 	}
 	else
 	{
-		TArray<uint8> TopMipData;
+		TArray64<uint8> TopMipData;
 		Texture->Source.GetMipData(TopMipData, 0);
 		Texture->Source.Init2DWithMipChain(SizeU, SizeV, Format);
 		int32 NumMips = Texture->Source.GetNumMips();
@@ -5803,14 +5803,14 @@ void ULandscapeComponent::GeneratePlatformVertexData(const ITargetPlatform* Targ
 	float HeightmapSubsectionOffsetV = (float)(SubsectionSizeVerts) / (float)GetHeightmap()->Source.GetSizeY();
 	
 	// Get the required mip data
-	TArray<TArray<uint8>> HeightmapMipRawData;
-	TArray<FColor*> HeightmapMipData;
+	TArray<TArray64<uint8>> HeightmapMipRawData;
+	TArray64<FColor*> HeightmapMipData;
 	for (int32 MipIdx = 0; MipIdx < FMath::Min(LANDSCAPE_MAX_ES_LOD, GetHeightmap()->Source.GetNumMips()); MipIdx++)
 	{
 		int32 MipSubsectionSizeVerts = (SubsectionSizeVerts) >> MipIdx;
 		if (MipSubsectionSizeVerts > 1)
 		{
-			new(HeightmapMipRawData) TArray<uint8>();
+			new(HeightmapMipRawData) TArray64<uint8>();
 			GetHeightmap()->Source.GetMipData(HeightmapMipRawData.Last(), MipIdx);
 			HeightmapMipData.Add((FColor*)HeightmapMipRawData.Last().GetData());
 		}
diff --git a/Engine/Source/Runtime/Landscape/Private/LandscapeEditLayers.cpp b/Engine/Source/Runtime/Landscape/Private/LandscapeEditLayers.cpp
--- a/Engine/Source/Runtime/Landscape/Private/LandscapeEditLayers.cpp
+++ b/Engine/Source/Runtime/Landscape/Private/LandscapeEditLayers.cpp
@@ -1669,7 +1669,7 @@ void ALandscape::CopyOldDataToDefaultLayer(ALandscapeProxy* InProxy)
 				check(DefaultLayerHeightmap != nullptr);
 
 				// Only copy Mip0 as other mips will get regenerated
-				TArray<uint8> ExistingMip0Data;
+				TArray64<uint8> ExistingMip0Data;
 				ComponentHeightmap->Source.GetMipData(ExistingMip0Data, 0);
 
 				FColor* Mip0Data = (FColor*)DefaultLayerHeightmap->Source.LockMip(0);
@@ -1724,7 +1724,7 @@ void ALandscape::CopyOldDataToDefaultLayer(ALandscapeProxy* InProxy)
 					UTexture2D* NewLayerWeightmapTexture = InProxy->CreateLandscapeTexture(ComponentWeightmap->Source.GetSizeX(), ComponentWeightmap->Source.GetSizeY(), TEXTUREGROUP_Terrain_Weightmap, ComponentWeightmap->Source.GetFormat());
 
 					// Only copy Mip0 as other mips will get regenerated
-					TArray<uint8> ExistingMip0Data;
+					TArray64<uint8> ExistingMip0Data;
 					ComponentWeightmap->Source.GetMipData(ExistingMip0Data, 0);
 
 					FColor* Mip0Data = (FColor*)NewLayerWeightmapTexture->Source.LockMip(0);
diff --git a/Engine/Source/Runtime/Landscape/Public/LandscapeDataAccess.h b/Engine/Source/Runtime/Landscape/Public/LandscapeDataAccess.h
--- a/Engine/Source/Runtime/Landscape/Public/LandscapeDataAccess.h
+++ b/Engine/Source/Runtime/Landscape/Public/LandscapeDataAccess.h
@@ -55,7 +55,7 @@ private:
 		:	LockCount(0)
 		{}
 
-		TArray<uint8> MipData;
+		TArray64<uint8> MipData;
 		int32 LockCount;
 	};
 
