F - Copy from 11287679
Support for Half (R16F) textures

diff --git a/Engine/Source/Developer/Android/AndroidTargetPlatform/Private/AndroidTargetPlatform.cpp b/Engine/Source/Developer/Android/AndroidTargetPlatform/Private/AndroidTargetPlatform.cpp
--- a/Engine/Source/Developer/Android/AndroidTargetPlatform/Private/AndroidTargetPlatform.cpp
+++ b/Engine/Source/Developer/Android/AndroidTargetPlatform/Private/AndroidTargetPlatform.cpp
@@ -506,6 +506,10 @@ void FAndroidTargetPlatform::GetTextureFormats( const UTexture* InTexture, TArra
 			{
 				FormatPerLayer[LayerIndex] = AndroidTexFormat::NameG8;
 			}
+			else if (LayerFormatSettings.CompressionSettings == TC_HalfFloat)
+			{
+				FormatPerLayer[LayerIndex] = AndroidTexFormat::NameR16F;
+			}
 			else if (LayerFormatSettings.CompressionSettings == TC_BC7)
 			{
 				if (!bIsCompressionValid) FormatPerLayer[LayerIndex] = AndroidTexFormat::NamePOTERROR;
@@ -557,6 +561,7 @@ void FAndroidTargetPlatform::GetAllTextureFormats(TArray<FName>& OutFormats) con
 	OutFormats.Add(AndroidTexFormat::NameG8);
 	OutFormats.Add(AndroidTexFormat::NameG8);
 	OutFormats.Add(AndroidTexFormat::NameG8);
+	OutFormats.Add(AndroidTexFormat::NameR16F);
 
 	auto AddAllTextureFormatIfSupports = [=, &OutFormats](bool bIsNonPOT)
 	{
diff --git a/Engine/Source/Developer/Android/AndroidTargetPlatform/Private/AndroidTargetPlatform.h b/Engine/Source/Developer/Android/AndroidTargetPlatform/Private/AndroidTargetPlatform.h
--- a/Engine/Source/Developer/Android/AndroidTargetPlatform/Private/AndroidTargetPlatform.h
+++ b/Engine/Source/Developer/Android/AndroidTargetPlatform/Private/AndroidTargetPlatform.h
@@ -64,6 +64,7 @@ namespace AndroidTexFormat
 	static FName NameG8(TEXT("G8"));
 	static FName NameVU8(TEXT("VU8"));
 	static FName NameRGBA16F(TEXT("RGBA16F"));
+	static FName NameR16F(TEXT("R16F"));
 
 	// Error "formats" (uncompressed)
 	static FName NamePOTERROR(TEXT("POTERROR"));
diff --git a/Engine/Source/Developer/Mac/MacTargetPlatform/Private/GenericMacTargetPlatform.h b/Engine/Source/Developer/Mac/MacTargetPlatform/Private/GenericMacTargetPlatform.h
--- a/Engine/Source/Developer/Mac/MacTargetPlatform/Private/GenericMacTargetPlatform.h
+++ b/Engine/Source/Developer/Mac/MacTargetPlatform/Private/GenericMacTargetPlatform.h
@@ -190,6 +190,7 @@ return TSuper::SupportsFeature(Feature);
 		static FName NameG16(TEXT("G16"));
 		static FName NameVU8(TEXT("VU8"));
 		static FName NameRGBA16F(TEXT("RGBA16F"));
+		static FName NameR16F(TEXT("R16F"));
 		static FName NameBC6H(TEXT("BC6H"));
 		static FName NameBC7(TEXT("BC7"));
 
@@ -264,6 +265,10 @@ return TSuper::SupportsFeature(Feature);
 		{
 			TextureFormatName = NameBC7;
 		}
+		else if (Settings == TC_HalfFloat)
+		{
+			TextureFormatName = NameR16F;
+		}
 		else if (bNoAlpha)
 		{
 			TextureFormatName = NameDXT1;
diff --git a/Engine/Source/Developer/TextureFormatUncompressed/Private/TextureFormatUncompressed.cpp b/Engine/Source/Developer/TextureFormatUncompressed/Private/TextureFormatUncompressed.cpp
--- a/Engine/Source/Developer/TextureFormatUncompressed/Private/TextureFormatUncompressed.cpp
+++ b/Engine/Source/Developer/TextureFormatUncompressed/Private/TextureFormatUncompressed.cpp
@@ -23,7 +23,8 @@ DEFINE_LOG_CATEGORY_STATIC(LogTextureFormatUncompressed, Log, All);
 	op(RGBA16F) \
 	op(XGXR8) \
 	op(RGBA8) \
-	op(POTERROR)
+	op(POTERROR) \
+	op(R16F)
 
 #define DECL_FORMAT_NAME(FormatName) static FName GTextureFormatName##FormatName = FName(TEXT(#FormatName));
 ENUM_SUPPORTED_FORMATS(DECL_FORMAT_NAME);
@@ -209,6 +210,19 @@ class FTextureFormatUncompressed : public ITextureFormat
 
 			return true;
 		}
+		else if (BuildSettings.TextureFormatName == GTextureFormatNameR16F)
+		{
+			FImage Image;
+			InImage.CopyTo(Image, ERawImageFormat::R16F, EGammaSpace::Linear);
+
+			OutCompressedImage.SizeX = Image.SizeX;
+			OutCompressedImage.SizeY = Image.SizeY;
+			OutCompressedImage.SizeZ = BuildSettings.bVolume ? Image.NumSlices : 1;
+			OutCompressedImage.PixelFormat = PF_R16F;
+			OutCompressedImage.RawData = Image.RawData;
+
+			return true;
+		}
 		else if (BuildSettings.TextureFormatName == GTextureFormatNamePOTERROR)
 		{
 			// load the error image data we will just repeat into the texture
diff --git a/Engine/Source/Developer/Windows/WindowsTargetPlatform/Private/GenericWindowsTargetPlatform.h b/Engine/Source/Developer/Windows/WindowsTargetPlatform/Private/GenericWindowsTargetPlatform.h
--- a/Engine/Source/Developer/Windows/WindowsTargetPlatform/Private/GenericWindowsTargetPlatform.h
+++ b/Engine/Source/Developer/Windows/WindowsTargetPlatform/Private/GenericWindowsTargetPlatform.h
@@ -280,6 +280,7 @@ public:
 		static FName NameG16(TEXT("G16"));
 		static FName NameVU8(TEXT("VU8"));
 		static FName NameRGBA16F(TEXT("RGBA16F"));
+		static FName NameR16F(TEXT("R16F"));
 		static FName NameBC6H(TEXT("BC6H"));
 		static FName NameBC7(TEXT("BC7"));
 
@@ -354,6 +355,10 @@ public:
 		{
 			TextureFormatName = NameBC7;
 		}
+		else if (Settings == TC_HalfFloat)
+		{
+			TextureFormatName = NameR16F;
+		}
 		else if (bNoAlpha)
 		{
 			TextureFormatName = NameDXT1;
diff --git a/Engine/Source/Runtime/Engine/Classes/Engine/Texture.h b/Engine/Source/Runtime/Engine/Classes/Engine/Texture.h
--- a/Engine/Source/Runtime/Engine/Classes/Engine/Texture.h
+++ b/Engine/Source/Runtime/Engine/Classes/Engine/Texture.h
@@ -39,6 +39,7 @@ enum TextureCompressionSettings
 	TC_DistanceFieldFont		UMETA(DisplayName="DistanceFieldFont (R8)"),
 	TC_HDR_Compressed			UMETA(DisplayName="HDRCompressed (RGB, BC6H, DX11)"),
 	TC_BC7						UMETA(DisplayName="BC7 (DX11, optional A)"),
+	TC_HalfFloat				UMETA(DisplayName="Half Float (R16F)"),
 	TC_MAX,
 };
 
diff --git a/Engine/Source/Runtime/Engine/Private/Texture.cpp b/Engine/Source/Runtime/Engine/Private/Texture.cpp
--- a/Engine/Source/Runtime/Engine/Private/Texture.cpp
+++ b/Engine/Source/Runtime/Engine/Private/Texture.cpp
@@ -242,7 +242,7 @@ void UTexture::PostEditChangeProperty(FPropertyChangedEvent& PropertyChangedEven
 		}
 #endif //WITH_EDITORONLY_DATA
 
-		bool bPreventSRGB = (CompressionSettings == TC_Alpha || CompressionSettings == TC_Normalmap || CompressionSettings == TC_Masks || CompressionSettings == TC_HDR || CompressionSettings == TC_HDR_Compressed);
+		bool bPreventSRGB = (CompressionSettings == TC_Alpha || CompressionSettings == TC_Normalmap || CompressionSettings == TC_Masks || CompressionSettings == TC_HDR || CompressionSettings == TC_HDR_Compressed || CompressionSettings == TC_HalfFloat);
 		if(bPreventSRGB && SRGB == true)
 		{
 			SRGB = false;
@@ -1308,7 +1308,7 @@ FName GetDefaultTextureFormatName( const ITargetPlatform* TargetPlatform, const
 	FName TextureFormatName = NAME_None;
 
 	/**
-	 * IF you add a format to this function don't forget to update GetAllDefaultTextureFormats 
+	 * IF you add a format to this function don't forget to update GetAllDefaultTextureFormatNames 
 	 */
 
 #if WITH_EDITOR
@@ -1326,6 +1326,7 @@ FName GetDefaultTextureFormatName( const ITargetPlatform* TargetPlatform, const
 	static FName NameG16(TEXT("G16"));
 	static FName NameVU8(TEXT("VU8"));
 	static FName NameRGBA16F(TEXT("RGBA16F"));
+	static FName NameR16F(TEXT("R16F"));
 	static FName NameBC6H(TEXT("BC6H"));
 	static FName NameBC7(TEXT("BC7"));
 
@@ -1430,6 +1431,10 @@ FName GetDefaultTextureFormatName( const ITargetPlatform* TargetPlatform, const
 	{
 		TextureFormatName = NameBC7;
 	}
+	else if (FormatSettings.CompressionSettings == TC_HalfFloat)
+	{
+		TextureFormatName = NameR16F;
+	}
 	else if (FormatSettings.CompressionNoAlpha)
 	{
 		TextureFormatName = NameDXT1;
@@ -1495,6 +1500,7 @@ void GetAllDefaultTextureFormats(const class ITargetPlatform* TargetPlatform, TA
 	static FName NameG16(TEXT("G16"));
 	static FName NameVU8(TEXT("VU8"));
 	static FName NameRGBA16F(TEXT("RGBA16F"));
+	static FName NameR16F(TEXT("R16F"));
 	static FName NameBC6H(TEXT("BC6H"));
 	static FName NameBC7(TEXT("BC7"));
 
@@ -1511,6 +1517,7 @@ void GetAllDefaultTextureFormats(const class ITargetPlatform* TargetPlatform, TA
 	OutFormats.Add(NameG16);
 	OutFormats.Add(NameVU8);
 	OutFormats.Add(NameRGBA16F);
+	OutFormats.Add(NameR16F);
 	if (bSupportDX11TextureFormats)
 	{
 		OutFormats.Add(NameBC6H);
diff --git a/Engine/Source/Runtime/Engine/Private/Texture2D.cpp b/Engine/Source/Runtime/Engine/Private/Texture2D.cpp
--- a/Engine/Source/Runtime/Engine/Private/Texture2D.cpp
+++ b/Engine/Source/Runtime/Engine/Private/Texture2D.cpp
@@ -1559,7 +1559,7 @@ FTexture2DResource::FTexture2DResource( UTexture2D* InOwner, int32 InitialMipCou
 	MipFadeSetting = (Owner->LODGroup == TEXTUREGROUP_Lightmap || Owner->LODGroup == TEXTUREGROUP_Shadowmap) ? MipFade_Slow : MipFade_Normal;
 
 	// HDR images are stored in linear but still require gamma correction to display correctly.
-	bIgnoreGammaConversions = !Owner->SRGB && Owner->CompressionSettings != TC_HDR && Owner->CompressionSettings != TC_HDR_Compressed;
+	bIgnoreGammaConversions = !Owner->SRGB && Owner->CompressionSettings != TC_HDR && Owner->CompressionSettings != TC_HDR_Compressed && Owner->CompressionSettings != TC_HalfFloat;
 	bSRGB = InOwner->SRGB;
 
 	check(InitialMipCount>0);
@@ -2104,7 +2104,7 @@ void FVirtualTexture2DResource::InitializeEditorResources(IVirtualTexture* InVir
 		RHIBindDebugLabelName(TextureRHI, *TextureOwner->GetName());
 		RHIUpdateTextureReference(TextureOwner->TextureReference.TextureReferenceRHI, TextureRHI);
 
-		bIgnoreGammaConversions = !TextureOwner->SRGB && TextureOwner->CompressionSettings != TC_HDR;
+		bIgnoreGammaConversions = !TextureOwner->SRGB && TextureOwner->CompressionSettings != TC_HDR && TextureOwner->CompressionSettings != TC_HalfFloat;
 
 		// re factored to ensure this is set earlier...make sure it's correct
 		ensure(bSRGB == TextureOwner->SRGB);
diff --git a/Engine/Source/Runtime/ImageCore/Private/ImageCore.cpp b/Engine/Source/Runtime/ImageCore/Private/ImageCore.cpp
--- a/Engine/Source/Runtime/ImageCore/Private/ImageCore.cpp
+++ b/Engine/Source/Runtime/ImageCore/Private/ImageCore.cpp
@@ -120,24 +120,43 @@ static void CopyImage(const FImage& SrcImage, FImage& DestImage)
 		case ERawImageFormat::RGBA16:
 			{
 				uint16* DestColors = DestImage.AsRGBA16();
-				for (int64 TexelIndex = 0; TexelIndex < NumTexels; ++TexelIndex)
+				ParallelFor(NumJobs, [DestColors, SrcColors, TexelsPerJob, NumTexels](int64 JobIndex)
 				{
-					int64 DestIndex = TexelIndex * 4;
-					DestColors[DestIndex + 0] = FMath::Clamp(FMath::FloorToInt(SrcColors[TexelIndex].R * 65535.999f), 0, 65535);
-					DestColors[DestIndex + 1] = FMath::Clamp(FMath::FloorToInt(SrcColors[TexelIndex].G * 65535.999f), 0, 65535);
-					DestColors[DestIndex + 2] = FMath::Clamp(FMath::FloorToInt(SrcColors[TexelIndex].B * 65535.999f), 0, 65535);
-					DestColors[DestIndex + 3] = FMath::Clamp(FMath::FloorToInt(SrcColors[TexelIndex].A * 65535.999f), 0, 65535);
-				}
+					for (int64 TexelIndex = 0; TexelIndex < NumTexels; ++TexelIndex)
+					{
+						int64 DestIndex = TexelIndex * 4;
+						DestColors[DestIndex + 0] = FMath::Clamp(FMath::FloorToInt(SrcColors[TexelIndex].R * 65535.999f), 0, 65535);
+						DestColors[DestIndex + 1] = FMath::Clamp(FMath::FloorToInt(SrcColors[TexelIndex].G * 65535.999f), 0, 65535);
+						DestColors[DestIndex + 2] = FMath::Clamp(FMath::FloorToInt(SrcColors[TexelIndex].B * 65535.999f), 0, 65535);
+						DestColors[DestIndex + 3] = FMath::Clamp(FMath::FloorToInt(SrcColors[TexelIndex].A * 65535.999f), 0, 65535);
+					}
+				});
 			}
 			break;
 		
 		case ERawImageFormat::RGBA16F:
 			{
 				FFloat16Color* DestColors = DestImage.AsRGBA16F();
-				for (int64 TexelIndex = 0; TexelIndex < NumTexels; ++TexelIndex)
+				ParallelFor(NumJobs, [DestColors, SrcColors, TexelsPerJob, NumTexels](int64 JobIndex)
 				{
-					DestColors[TexelIndex] = FFloat16Color(SrcColors[TexelIndex]);
-				}
+					for (int64 TexelIndex = 0; TexelIndex < NumTexels; ++TexelIndex)
+					{
+						DestColors[TexelIndex] = FFloat16Color(SrcColors[TexelIndex]);
+					}
+				});
+			}
+			break;
+
+		case ERawImageFormat::R16F:
+			{
+				FFloat16* DestColors = DestImage.AsR16F();
+				ParallelFor(NumJobs, [DestColors, SrcColors, TexelsPerJob, NumTexels](int64 JobIndex)
+				{
+					for (int64 TexelIndex = 0; TexelIndex < NumTexels; ++TexelIndex)
+					{
+						DestColors[TexelIndex] = FFloat16(SrcColors[TexelIndex].R);
+					}
+				});
 			}
 			break;
 		}
@@ -253,6 +272,16 @@ static void CopyImage(const FImage& SrcImage, FImage& DestImage)
 				}
 			}
 			break;
+
+		case ERawImageFormat::R16F:
+		{
+			const FFloat16* SrcColors = SrcImage.AsR16F();
+			for (int64 TexelIndex = 0; TexelIndex < NumTexels; ++TexelIndex)
+			{
+				DestColors[TexelIndex] = FLinearColor(SrcColors[TexelIndex].GetFloat(), 0, 0, 1);
+			}
+		}
+		break;
 		}
 	}
 	else
@@ -414,6 +443,7 @@ int32 FImage::GetBytesPerPixel() const
 		break;
 	
 	case ERawImageFormat::G16:
+	case ERawImageFormat::R16F:
 		BytesPerPixel = 2;
 		break;
 
diff --git a/Engine/Source/Runtime/ImageCore/Public/ImageCore.h b/Engine/Source/Runtime/ImageCore/Public/ImageCore.h
--- a/Engine/Source/Runtime/ImageCore/Public/ImageCore.h
+++ b/Engine/Source/Runtime/ImageCore/Public/ImageCore.h
@@ -24,7 +24,8 @@ namespace ERawImageFormat
 		RGBA16,
 		RGBA16F,
 		RGBA32F,
-		G16
+		G16,
+		R16F,
 	};
 };
 
@@ -179,6 +180,12 @@ public:
 		return (struct FLinearColor*)RawData.GetData();
 	}
 
+	class FFloat16* AsR16F()
+	{
+		check(Format == ERawImageFormat::R16F);
+		return (class FFloat16*)RawData.GetData();
+	}
+
 	// Convenience accessors to const raw data
 
 	const uint8* AsG8() const
@@ -223,6 +230,12 @@ public:
 		return (struct FLinearColor*)RawData.GetData();
 	}
 
+	const class FFloat16* AsR16F() const
+	{
+		check(Format == ERawImageFormat::R16F);
+		return (const class FFloat16*)RawData.GetData();
+	}
+
 	FORCEINLINE bool IsGammaCorrected() const
 	{
 		return GammaSpace != EGammaSpace::Linear;
