fix gamma correction to pow22 doesn't and shouldn't exist

diff --git a/Engine/Source/Runtime/ImageCore/Private/ImageCore.cpp b/Engine/Source/Runtime/ImageCore/Private/ImageCore.cpp
--- a/Engine/Source/Runtime/ImageCore/Private/ImageCore.cpp
+++ b/Engine/Source/Runtime/ImageCore/Private/ImageCore.cpp
@@ -57,6 +57,14 @@ static void CopyImage(const FImage& SrcImage, FImage& DestImage)
 		// Convert from 32-bit linear floating point.
 		TArrayView64<const FLinearColor> SrcColors = SrcImage.AsRGBA32F();
 	
+		// if gamma correction is done, it's always to sRGB , not to Pow22
+		// so if Pow22 was requested, change to sRGB
+		// so that Float->int->Float roundtrips correctly
+		if ( DestImage.GammaSpace == EGammaSpace::Pow22 )
+		{
+			DestImage.GammaSpace = EGammaSpace::sRGB;
+		}
+
 		switch (DestImage.Format)
 		{
 		case ERawImageFormat::G8:
