Fix ParallelFor usage in CopyImage, all conversion loops now use JobIndex properly.
Remove TArrayView64 to avoid per-pixel bounds checks in conversion loops.

diff --git a/Engine/Source/Runtime/ImageCore/Private/ImageCore.cpp b/Engine/Source/Runtime/ImageCore/Private/ImageCore.cpp
--- a/Engine/Source/Runtime/ImageCore/Private/ImageCore.cpp
+++ b/Engine/Source/Runtime/ImageCore/Private/ImageCore.cpp
@@ -10,11 +10,14 @@ DEFINE_LOG_CATEGORY_STATIC(LogImageCore, Log, All);
 
 static constexpr float MAX_HALF_FLOAT16 = 65504.0f;
 
-FORCEINLINE static void SaturateToHalfFloat(FLinearColor& LinearCol)
+FORCEINLINE static FLinearColor SaturateToHalfFloat(const FLinearColor& LinearCol)
 {
-	LinearCol.R = FMath::Clamp(LinearCol.R, -MAX_HALF_FLOAT16, MAX_HALF_FLOAT16);
-	LinearCol.G = FMath::Clamp(LinearCol.G, -MAX_HALF_FLOAT16, MAX_HALF_FLOAT16);
-	LinearCol.B = FMath::Clamp(LinearCol.B, -MAX_HALF_FLOAT16, MAX_HALF_FLOAT16);
+	FLinearColor Result;
+	Result.R = FMath::Clamp(LinearCol.R, -MAX_HALF_FLOAT16, MAX_HALF_FLOAT16);
+	Result.G = FMath::Clamp(LinearCol.G, -MAX_HALF_FLOAT16, MAX_HALF_FLOAT16);
+	Result.B = FMath::Clamp(LinearCol.B, -MAX_HALF_FLOAT16, MAX_HALF_FLOAT16);
+	Result.A = LinearCol.A;
+	return Result;
 }
 
 IMPLEMENT_MODULE(FDefaultModuleImpl, ImageCore);
@@ -86,6 +89,20 @@ IMAGECORE_API int32 ImageParallelForComputeNumJobsForRows(int32 & OutNumItemsPer
 	return NumJobs;
 }
 
+template <typename Lambda>
+static void ParallelLoop(const TCHAR* DebugName, int32 NumJobs, int64 TexelsPerJob, int64 NumTexels, const Lambda& Func)
+{
+	ParallelFor(DebugName, NumJobs, 1, [=](int64 JobIndex)
+	{
+		const int64 StartIndex = JobIndex * TexelsPerJob;
+		const int64 EndIndex = FMath::Min(StartIndex + TexelsPerJob, NumTexels);
+		for (int64 TexelIndex = StartIndex; TexelIndex < EndIndex; ++TexelIndex)
+		{
+			Func(TexelIndex);
+		}
+	});
+}
+
 /**
  * Copies an image accounting for format differences. Sizes must match.
  *
@@ -113,7 +130,7 @@ static void CopyImage(const FImage& SrcImage, FImage& DestImage)
 	else if (SrcImage.Format == ERawImageFormat::RGBA32F)
 	{
 		// Convert from 32-bit linear floating point.
-		TArrayView64<const FLinearColor> SrcColors = SrcImage.AsRGBA32F();
+		const FLinearColor* SrcColors = SrcImage.AsRGBA32F().GetData();
 	
 		// if gamma correction is done, it's always to sRGB , not to Pow22
 		// so if Pow22 was requested, change to sRGB
@@ -127,103 +144,75 @@ static void CopyImage(const FImage& SrcImage, FImage& DestImage)
 		{
 		case ERawImageFormat::G8:
 			{
-				TArrayView64<uint8> DestLum = DestImage.AsG8();
-				ParallelFor( TEXT("Texture.CopyImage.PF"),NumJobs,1, [NumJobs, DestLum, SrcColors, TexelsPerJob, NumTexels, bDestIsGammaCorrected](int64 JobIndex)
+				uint8* DestLum = DestImage.AsG8().GetData();
+				ParallelLoop(TEXT("Texture.CopyImage.PF"), NumJobs, TexelsPerJob, NumTexels, [DestLum, SrcColors, bDestIsGammaCorrected](int64 TexelIndex)
 				{
-					const int64 StartIndex = JobIndex * TexelsPerJob;
-					const int64 EndIndex = FMath::Min(StartIndex + TexelsPerJob, NumTexels);
-					for (int64 TexelIndex = StartIndex; TexelIndex < EndIndex; ++TexelIndex)
-					{
-						DestLum[TexelIndex] = SrcColors[TexelIndex].ToFColor(bDestIsGammaCorrected).R;
-					}
+					DestLum[TexelIndex] = SrcColors[TexelIndex].ToFColor(bDestIsGammaCorrected).R;
 				});
 			}
 			break;
 
 		case ERawImageFormat::G16:
-		{
-			TArrayView64<uint16>DestLum = DestImage.AsG16();
-			ParallelFor(TEXT("Texture.CopyImage.PF"),NumJobs,1, [NumJobs, DestLum, SrcColors, TexelsPerJob, NumTexels, bDestIsGammaCorrected](int64 JobIndex)
 			{
-				const int64 StartIndex = JobIndex * TexelsPerJob;
-				const int64 EndIndex = FMath::Min(StartIndex + TexelsPerJob, NumTexels);
-				for (int64 TexelIndex = StartIndex; TexelIndex < EndIndex; ++TexelIndex)
+				uint16* DestLum = DestImage.AsG16().GetData();
+				ParallelLoop(TEXT("Texture.CopyImage.PF"), NumJobs, TexelsPerJob, NumTexels, [DestLum, SrcColors](int64 TexelIndex)
 				{
-					DestLum[TexelIndex] = FColor::QuantizeUNormFloatTo16( SrcColors[TexelIndex].R );
-				}
-			});
-		}
-		break;
+					DestLum[TexelIndex] = FColor::QuantizeUNormFloatTo16(SrcColors[TexelIndex].R);
+				});
+			}
+			break;
 
 		case ERawImageFormat::BGRA8:
 			{
-				TArrayView64<FColor> DestColors = DestImage.AsBGRA8();
-				ParallelFor(TEXT("Texture.CopyImage.PF"),NumJobs,1, [DestColors, SrcColors, TexelsPerJob, NumTexels, bDestIsGammaCorrected](int64 JobIndex)
+				FColor* DestColors = DestImage.AsBGRA8().GetData();
+				ParallelLoop(TEXT("Texture.CopyImage.PF"), NumJobs, TexelsPerJob, NumTexels, [DestColors, SrcColors, bDestIsGammaCorrected](int64 TexelIndex)
 				{
-					const int64 StartIndex = JobIndex * TexelsPerJob;
-					const int64 EndIndex = FMath::Min(StartIndex + TexelsPerJob, NumTexels);
-					for (int64 TexelIndex = StartIndex; TexelIndex < EndIndex; ++TexelIndex)
-					{
-						DestColors[TexelIndex] = SrcColors[TexelIndex].ToFColor(bDestIsGammaCorrected);
-					}
+					DestColors[TexelIndex] = SrcColors[TexelIndex].ToFColor(bDestIsGammaCorrected);
 				});
 			}
 			break;
 		
 		case ERawImageFormat::BGRE8:
 			{
-				TArrayView64<FColor> DestColors = DestImage.AsBGRE8();
-				ParallelFor(TEXT("Texture.CopyImage.PF"),NumJobs,1, [DestColors, SrcColors, TexelsPerJob, NumTexels, bDestIsGammaCorrected](int64 JobIndex)
+				FColor* DestColors = DestImage.AsBGRE8().GetData();
+				ParallelLoop(TEXT("Texture.CopyImage.PF"), NumJobs, TexelsPerJob, NumTexels, [DestColors, SrcColors](int64 TexelIndex)
 				{
-					const int64 StartIndex = JobIndex * TexelsPerJob;
-					const int64 EndIndex = FMath::Min(StartIndex + TexelsPerJob, NumTexels);
-					for (int64 TexelIndex = StartIndex; TexelIndex < EndIndex; ++TexelIndex)
-					{
-						DestColors[TexelIndex] = SrcColors[TexelIndex].ToRGBE();
-					}
+					DestColors[TexelIndex] = SrcColors[TexelIndex].ToRGBE();
 				});
 			}
 			break;
 		
 		case ERawImageFormat::RGBA16:
 			{
-				TArrayView64<uint16> DestColors = DestImage.AsRGBA16();
-				ParallelFor(TEXT("Texture.CopyImage.PF"),NumJobs,1, [DestColors, SrcColors, TexelsPerJob, NumTexels](int64 JobIndex)
+				uint16* DestColors = DestImage.AsRGBA16().GetData();
+				ParallelLoop(TEXT("Texture.CopyImage.PF"), NumJobs, TexelsPerJob, NumTexels, [DestColors, SrcColors](int64 TexelIndex)
 				{
-					for (int64 TexelIndex = 0; TexelIndex < NumTexels; ++TexelIndex)
-					{
-						int64 DestIndex = TexelIndex * 4;
-						DestColors[DestIndex + 0] = FColor::QuantizeUNormFloatTo16( SrcColors[TexelIndex].R );
-						DestColors[DestIndex + 1] = FColor::QuantizeUNormFloatTo16( SrcColors[TexelIndex].G );
-						DestColors[DestIndex + 2] = FColor::QuantizeUNormFloatTo16( SrcColors[TexelIndex].B );
-						DestColors[DestIndex + 3] = FColor::QuantizeUNormFloatTo16( SrcColors[TexelIndex].A );
-					}
+					FLinearColor Src = SrcColors[TexelIndex];
+					uint16* Dst = DestColors + TexelIndex * 4;
+					Dst[0] = FColor::QuantizeUNormFloatTo16(Src.R);
+					Dst[1] = FColor::QuantizeUNormFloatTo16(Src.G);
+					Dst[2] = FColor::QuantizeUNormFloatTo16(Src.B);
+					Dst[3] = FColor::QuantizeUNormFloatTo16(Src.A);
 				});
 			}
 			break;
 		
 		case ERawImageFormat::RGBA16F:
 			{
-				TArrayView64<FFloat16Color> DestColors = DestImage.AsRGBA16F();
-				ParallelFor(TEXT("Texture.CopyImage.PF"),NumJobs,1, [DestColors, SrcColors, TexelsPerJob, NumTexels](int64 JobIndex)
+				FFloat16Color* DestColors = DestImage.AsRGBA16F().GetData();
+				ParallelLoop(TEXT("Texture.CopyImage.PF"), NumJobs, TexelsPerJob, NumTexels, [DestColors, SrcColors](int64 TexelIndex)
 				{
-					for (int64 TexelIndex = 0; TexelIndex < NumTexels; ++TexelIndex)
-					{
-						DestColors[TexelIndex] = FFloat16Color(SrcColors[TexelIndex]);
-					}
+					DestColors[TexelIndex] = FFloat16Color(SrcColors[TexelIndex]);
 				});
 			}
 			break;
 
 		case ERawImageFormat::R16F:
 			{
-				TArrayView64<FFloat16> DestColors = DestImage.AsR16F();
-				ParallelFor(TEXT("Texture.CopyImage.PF"),NumJobs,1, [DestColors, SrcColors, TexelsPerJob, NumTexels](int64 JobIndex)
+				FFloat16* DestColors = DestImage.AsR16F().GetData();
+				ParallelLoop(TEXT("Texture.CopyImage.PF"), NumJobs, TexelsPerJob, NumTexels, [DestColors, SrcColors](int64 TexelIndex)
 				{
-					for (int64 TexelIndex = 0; TexelIndex < NumTexels; ++TexelIndex)
-					{
-						DestColors[TexelIndex] = FFloat16(SrcColors[TexelIndex].R);
-					}
+					DestColors[TexelIndex] = FFloat16(SrcColors[TexelIndex].R);
 				});
 			}
 			break;
@@ -232,92 +221,94 @@ static void CopyImage(const FImage& SrcImage, FImage& DestImage)
 	else if (DestImage.Format == ERawImageFormat::RGBA32F)
 	{
 		// Convert to 32-bit linear floating point.
-		TArrayView64<FLinearColor> DestColors = DestImage.AsRGBA32F();
+		FLinearColor* DestColors = DestImage.AsRGBA32F().GetData();
 		switch (SrcImage.Format)
 		{
 		case ERawImageFormat::G8:
 			{
-				TArrayView64<const uint8> SrcLum = SrcImage.AsG8();
-				for (int64 TexelIndex = 0; TexelIndex < NumTexels; ++TexelIndex)
+				const uint8* SrcLum = SrcImage.AsG8().GetData();
+				switch (SrcImage.GammaSpace)
 				{
-					FColor SrcColor(SrcLum[TexelIndex],SrcLum[TexelIndex],SrcLum[TexelIndex],255);
-					
-					switch ( SrcImage.GammaSpace )
+				case EGammaSpace::Linear:
+					ParallelLoop(TEXT("Texture.CopyImage.PF"), NumJobs, TexelsPerJob, NumTexels, [DestColors, SrcLum](int64 TexelIndex)
 					{
-					case EGammaSpace::Linear:
+						uint8 Lum = SrcLum[TexelIndex];
+						FColor SrcColor(Lum, Lum, Lum, 255);
 						DestColors[TexelIndex] = SrcColor.ReinterpretAsLinear();
-						break;
-					case EGammaSpace::sRGB:
-						DestColors[TexelIndex] = FLinearColor(SrcColor);
-						break;
-					case EGammaSpace::Pow22:
+					});
+					break;
+				case EGammaSpace::sRGB:
+					ParallelLoop(TEXT("Texture.CopyImage.PF"), NumJobs, TexelsPerJob, NumTexels, [DestColors, SrcLum](int64 TexelIndex)
+					{
+							uint8 Lum = SrcLum[TexelIndex];
+							FColor SrcColor(Lum, Lum, Lum, 255);
+							DestColors[TexelIndex] = FLinearColor(SrcColor);
+					});
+					break;
+				case EGammaSpace::Pow22:
+					ParallelLoop(TEXT("Texture.CopyImage.PF"), NumJobs, TexelsPerJob, NumTexels, [DestColors, SrcLum](int64 TexelIndex)
+					{
+						uint8 Lum = SrcLum[TexelIndex];
+						FColor SrcColor(Lum, Lum, Lum, 255);
 						DestColors[TexelIndex] = FLinearColor::FromPow22Color(SrcColor);
-						break;
-					}
+					});
+					break;
 				}
 			}
 			break;
 
 		case ERawImageFormat::G16:
-		{
-			TArrayView64<const uint16> SrcLum = SrcImage.AsG16();
-			for (int64 TexelIndex = 0; TexelIndex < NumTexels; ++TexelIndex)
 			{
-				float Value = FColor::DequantizeUNorm16ToFloat( SrcLum[TexelIndex] );
-				DestColors[TexelIndex] = FLinearColor(Value,Value,Value, 1.0f);
+				const uint16* SrcLum = SrcImage.AsG16().GetData();
+				ParallelLoop(TEXT("Texture.CopyImage.PF"), NumJobs, TexelsPerJob, NumTexels, [DestColors, SrcLum](int64 TexelIndex)
+				{
+					float Value = FColor::DequantizeUNorm16ToFloat(SrcLum[TexelIndex]);
+					DestColors[TexelIndex] = FLinearColor(Value, Value, Value, 1.0f);
+				});
 			}
-		}
-		break;
+			break;
 
 		case ERawImageFormat::BGRA8:
 			{
-				TArrayView64<const FColor> SrcColors = SrcImage.AsBGRA8();
-				EGammaSpace SrcGamma = SrcImage.GammaSpace;
-
-				ParallelFor(TEXT("Texture.CopyImage.PF"),NumJobs,1, [DestColors, SrcColors, TexelsPerJob, NumTexels, SrcGamma](int64 JobIndex)
+				const FColor* SrcColors = SrcImage.AsBGRA8().GetData();
+				switch (SrcImage.GammaSpace)
 				{
-					int64 StartIndex = JobIndex * TexelsPerJob;
-					int64 EndIndex = FMath::Min(StartIndex + TexelsPerJob, NumTexels);
-
-					switch ( SrcGamma )
+				case EGammaSpace::Linear:
+					ParallelLoop(TEXT("Texture.CopyImage.PF"), NumJobs, TexelsPerJob, NumTexels, [DestColors, SrcColors](int64 TexelIndex)
 					{
-					case EGammaSpace::Linear:
-						for (int64 TexelIndex = StartIndex; TexelIndex < EndIndex; ++TexelIndex)
-						{
-							DestColors[TexelIndex] = SrcColors[TexelIndex].ReinterpretAsLinear();
-						}
+						DestColors[TexelIndex] = SrcColors[TexelIndex].ReinterpretAsLinear();
+					});
 					break;
-					case EGammaSpace::sRGB:
-						for (int64 TexelIndex = StartIndex; TexelIndex < EndIndex; ++TexelIndex)
-						{
-							DestColors[TexelIndex] = FLinearColor(SrcColors[TexelIndex]);
-						}
+				case EGammaSpace::sRGB:
+					ParallelLoop(TEXT("Texture.CopyImage.PF"), NumJobs, TexelsPerJob, NumTexels, [DestColors, SrcColors](int64 TexelIndex)
+					{
+						DestColors[TexelIndex] = FLinearColor(SrcColors[TexelIndex]);
+					});
 					break;
-					case EGammaSpace::Pow22:
-						for (int64 TexelIndex = StartIndex; TexelIndex < EndIndex; ++TexelIndex)
-						{
-							DestColors[TexelIndex] = FLinearColor::FromPow22Color(SrcColors[TexelIndex]);
-						}
+				case EGammaSpace::Pow22:
+					ParallelLoop(TEXT("Texture.CopyImage.PF"), NumJobs, TexelsPerJob, NumTexels, [DestColors, SrcColors](int64 TexelIndex)
+					{
+						DestColors[TexelIndex] = FLinearColor::FromPow22Color(SrcColors[TexelIndex]);
+					});
 					break;
-					}
-				});
+				}
 			}
 			break;
 
 		case ERawImageFormat::BGRE8:
 			{
-				TArrayView64<const FColor> SrcColors = SrcImage.AsBGRE8();
-				for (int64 TexelIndex = 0; TexelIndex < NumTexels; ++TexelIndex)
+				const FColor* SrcColors = SrcImage.AsBGRE8().GetData();
+				ParallelLoop(TEXT("Texture.CopyImage.PF"), NumJobs, TexelsPerJob, NumTexels, [DestColors, SrcColors](int64 TexelIndex)
 				{
 					DestColors[TexelIndex] = SrcColors[TexelIndex].FromRGBE();
-				}
+				});
 			}
 			break;
 
 		case ERawImageFormat::RGBA16:
 			{
-				TArrayView64<const uint16> SrcColors = SrcImage.AsRGBA16();
-				for (int64 TexelIndex = 0; TexelIndex < NumTexels; ++TexelIndex)
+				const uint16* SrcColors = SrcImage.AsRGBA16().GetData();
+				ParallelLoop(TEXT("Texture.CopyImage.PF"), NumJobs, TexelsPerJob, NumTexels, [DestColors, SrcColors](int64 TexelIndex)
 				{
 					int64 SrcIndex = TexelIndex * 4;
 					DestColors[TexelIndex] = FLinearColor(
@@ -325,30 +316,30 @@ static void CopyImage(const FImage& SrcImage, FImage& DestImage)
 						SrcColors[SrcIndex + 1] / 65535.0f,
 						SrcColors[SrcIndex + 2] / 65535.0f,
 						SrcColors[SrcIndex + 3] / 65535.0f
-						);
-				}
+					);
+				});
 			}
 			break;
 
 		case ERawImageFormat::RGBA16F:
 			{
-				TArrayView64<const FFloat16Color> SrcColors = SrcImage.AsRGBA16F();
-				for (int64 TexelIndex = 0; TexelIndex < NumTexels; ++TexelIndex)
+				const FFloat16Color* SrcColors = SrcImage.AsRGBA16F().GetData();
+				ParallelLoop(TEXT("Texture.CopyImage.PF"), NumJobs, TexelsPerJob, NumTexels, [DestColors, SrcColors](int64 TexelIndex)
 				{
 					DestColors[TexelIndex] = SrcColors[TexelIndex].GetFloats();
-				}
+				});
 			}
 			break;
 
 		case ERawImageFormat::R16F:
-		{
-			TArrayView64<const FFloat16> SrcColors = SrcImage.AsR16F();
-			for (int64 TexelIndex = 0; TexelIndex < NumTexels; ++TexelIndex)
 			{
-				DestColors[TexelIndex] = FLinearColor(SrcColors[TexelIndex].GetFloat(), 0, 0, 1);
+				const FFloat16* SrcColors = SrcImage.AsR16F().GetData();
+				ParallelLoop(TEXT("Texture.CopyImage.PF"), NumJobs, TexelsPerJob, NumTexels, [DestColors, SrcColors](int64 TexelIndex)
+				{
+					DestColors[TexelIndex] = FLinearColor(SrcColors[TexelIndex].GetFloat(), 0, 0, 1);
+				});
 			}
-		}
-		break;
+			break;
 		}
 	}
 	else
@@ -379,26 +370,21 @@ void FImage::TransformToWorkingColorSpace(const FVector2D& SourceRedChromaticity
 
 	UE::Color::FColorSpaceTransform Transform(Source, Target, Method);
 
-	TArrayView64<FLinearColor> ImageColors = AsRGBA32F();
+	FLinearColor* ImageColors = AsRGBA32F().GetData();
 
 	const int64 NumTexels = int64(SizeX) * SizeY * NumSlices;
 	int64 TexelsPerJob;
 	int32 NumJobs = ImageParallelForComputeNumJobsForPixels(TexelsPerJob,NumTexels);
 
-	ParallelFor(TEXT("Texture.TransformToWorkingColorSpace.PF"),NumJobs,1, [Transform, ImageColors, TexelsPerJob, NumTexels](int64 JobIndex)
-		{
-
-			const int64 StartIndex = JobIndex * TexelsPerJob;
-			const int64 EndIndex = FMath::Min(StartIndex + TexelsPerJob, NumTexels);
-			for (int64 TexelIndex = StartIndex; TexelIndex < EndIndex; ++TexelIndex)
-			{
-				ImageColors[TexelIndex] = Transform.Apply(ImageColors[TexelIndex]);
-				SaturateToHalfFloat(ImageColors[TexelIndex]);
-			}
-		});
+	ParallelLoop(TEXT("Texture.TransformToWorkingColorSpace.PF"), NumJobs, TexelsPerJob, NumTexels, [Transform, ImageColors](int64 TexelIndex)
+	{
+		FLinearColor Color = ImageColors[TexelIndex];
+		Color = Transform.Apply(Color);
+		ImageColors[TexelIndex] = SaturateToHalfFloat(Color);
+	});
 }
 
-static FLinearColor SampleImage(TArrayView64<const FLinearColor> Pixels, int Width, int Height, float X, float Y)
+static FLinearColor SampleImage(const FLinearColor* Pixels, int Width, int Height, float X, float Y)
 {
 	const int64 TexelX0 = FMath::FloorToInt(X);
 	const int64 TexelY0 = FMath::FloorToInt(Y);
@@ -424,8 +410,8 @@ static FLinearColor SampleImage(TArrayView64<const FLinearColor> Pixels, int Wid
 
 static void ResizeImage(const FImage& SrcImage, FImage& DestImage)
 {
-	TArrayView64<const FLinearColor> SrcPixels = SrcImage.AsRGBA32F();
-	TArrayView64<FLinearColor> DestPixels = DestImage.AsRGBA32F();
+	const FLinearColor* SrcPixels = SrcImage.AsRGBA32F().GetData();
+	FLinearColor* DestPixels = DestImage.AsRGBA32F().GetData();
 	const float DestToSrcScaleX = (float)SrcImage.SizeX / (float)DestImage.SizeX;
 	const float DestToSrcScaleY = (float)SrcImage.SizeY / (float)DestImage.SizeY;
 
@@ -578,119 +564,112 @@ void FImage::Linearize(uint8 SourceEncoding, FImage& DestImage) const
 	check(DecodeFunction != nullptr);
 
 	// Convert to 32-bit linear floating point.
-	TArrayView64<FLinearColor> DestColors = DestImage.AsRGBA32F();
+	FLinearColor* DestColors = DestImage.AsRGBA32F().GetData();
 	switch (SrcImage.Format)
 	{
-	case ERawImageFormat::G8:
-	{
-		TArrayView64<const uint8> SrcLum = SrcImage.AsG8();
-		for (int64 TexelIndex = 0; TexelIndex < NumTexels; ++TexelIndex)
+		case ERawImageFormat::G8:
 		{
-			FColor SrcColor(SrcLum[TexelIndex], SrcLum[TexelIndex], SrcLum[TexelIndex], 255);
-			DestColors[TexelIndex] = SrcColor.ReinterpretAsLinear();
-			DestColors[TexelIndex] = DecodeFunction(DestColors[TexelIndex]);
-			SaturateToHalfFloat(DestColors[TexelIndex]);
+			const uint8* SrcLum = SrcImage.AsG8().GetData();
+			ParallelLoop(TEXT("Texture.Linearize.PF"), NumJobs, TexelsPerJob, NumTexels, [DestColors, SrcLum, DecodeFunction](int64 TexelIndex)
+			{
+				FColor SrcColor(SrcLum[TexelIndex], SrcLum[TexelIndex], SrcLum[TexelIndex], 255);
+				FLinearColor Color = SrcColor.ReinterpretAsLinear();
+				Color = DecodeFunction(Color);
+				DestColors[TexelIndex] = SaturateToHalfFloat(Color);
+			});
+			break;
 		}
-	}
-	break;
 
-	case ERawImageFormat::G16:
-	{
-		TArrayView64<const uint16> SrcLum = SrcImage.AsG16();
-		for (int64 TexelIndex = 0; TexelIndex < NumTexels; ++TexelIndex)
+		case ERawImageFormat::G16:
 		{
-			float Value = FColor::DequantizeUNorm16ToFloat(SrcLum[TexelIndex]);
-			DestColors[TexelIndex] = FLinearColor(Value, Value, Value, 1.0f);
-			DestColors[TexelIndex] = DecodeFunction(DestColors[TexelIndex]);
-			SaturateToHalfFloat(DestColors[TexelIndex]);
+			const uint16* SrcLum = SrcImage.AsG16().GetData();
+			ParallelLoop(TEXT("Texture.Linearize.PF"), NumJobs, TexelsPerJob, NumTexels, [DestColors, SrcLum, DecodeFunction](int64 TexelIndex)
+			{
+				float Value = FColor::DequantizeUNorm16ToFloat(SrcLum[TexelIndex]);
+				FLinearColor Color = FLinearColor(Value, Value, Value, 1.0f);
+				Color = DecodeFunction(Color);
+				DestColors[TexelIndex] = SaturateToHalfFloat(Color);
+			});
+			break;
 		}
-	}
-	break;
-
-	case ERawImageFormat::BGRA8:
-	{
-		TArrayView64<const FColor> SrcColors = SrcImage.AsBGRA8();
-		EGammaSpace SrcGamma = SrcImage.GammaSpace;
 
-		ParallelFor(TEXT("Texture.Linearize.PF"),NumJobs,1, [DestColors, SrcColors, TexelsPerJob, NumTexels, SrcGamma, DecodeFunction](int64 JobIndex)
+		case ERawImageFormat::BGRA8:
+		{
+			const FColor* SrcColors = SrcImage.AsBGRA8().GetData();
+			ParallelLoop(TEXT("Texture.Linearize.PF"), NumJobs, TexelsPerJob, NumTexels, [DestColors, SrcColors, DecodeFunction](int64 TexelIndex)
 			{
-				int64 StartIndex = JobIndex * TexelsPerJob;
-				int64 EndIndex = FMath::Min(StartIndex + TexelsPerJob, NumTexels);
-				for (int64 TexelIndex = StartIndex; TexelIndex < EndIndex; ++TexelIndex)
-				{
-					DestColors[TexelIndex] = SrcColors[TexelIndex].ReinterpretAsLinear();
-					DestColors[TexelIndex] = DecodeFunction(DestColors[TexelIndex]);
-					SaturateToHalfFloat(DestColors[TexelIndex]);
-				}
+				FLinearColor Color = SrcColors[TexelIndex].ReinterpretAsLinear();
+				Color = DecodeFunction(Color);
+				DestColors[TexelIndex] = SaturateToHalfFloat(Color);
 			});
-	}
-	break;
+			break;
+		}
 
-	case ERawImageFormat::BGRE8:
-	{
-		TArrayView64<const FColor> SrcColors = SrcImage.AsBGRE8();
-		for (int64 TexelIndex = 0; TexelIndex < NumTexels; ++TexelIndex)
+		case ERawImageFormat::BGRE8:
 		{
-			DestColors[TexelIndex] = SrcColors[TexelIndex].FromRGBE();
-			DestColors[TexelIndex] = DecodeFunction(DestColors[TexelIndex]);
-			SaturateToHalfFloat(DestColors[TexelIndex]);
+			const FColor* SrcColors = SrcImage.AsBGRE8().GetData();
+			ParallelLoop(TEXT("Texture.Linearize.PF"), NumJobs, TexelsPerJob, NumTexels, [DestColors, SrcColors, DecodeFunction](int64 TexelIndex)
+			{
+				FLinearColor Color = SrcColors[TexelIndex].FromRGBE();
+				Color = DecodeFunction(Color);
+				DestColors[TexelIndex] = SaturateToHalfFloat(Color);
+			});
+			break;
 		}
-	}
-	break;
 
-	case ERawImageFormat::RGBA16:
-	{
-		TArrayView64<const uint16> SrcColors = SrcImage.AsRGBA16();
-		for (int64 TexelIndex = 0; TexelIndex < NumTexels; ++TexelIndex)
+		case ERawImageFormat::RGBA16:
 		{
-			int64 SrcIndex = TexelIndex * 4;
-			DestColors[TexelIndex] = FLinearColor(
-				SrcColors[SrcIndex + 0] / 65535.0f,
-				SrcColors[SrcIndex + 1] / 65535.0f,
-				SrcColors[SrcIndex + 2] / 65535.0f,
-				SrcColors[SrcIndex + 3] / 65535.0f
-			);
-			DestColors[TexelIndex] = DecodeFunction(DestColors[TexelIndex]);
-			SaturateToHalfFloat(DestColors[TexelIndex]);
+			const uint16* SrcColors = SrcImage.AsRGBA16().GetData();
+			ParallelLoop(TEXT("Texture.Linearize.PF"), NumJobs, TexelsPerJob, NumTexels, [DestColors, SrcColors, DecodeFunction](int64 TexelIndex)
+			{
+				int64 SrcIndex = TexelIndex * 4;
+				FLinearColor Color(
+					SrcColors[SrcIndex + 0] / 65535.0f,
+					SrcColors[SrcIndex + 1] / 65535.0f,
+					SrcColors[SrcIndex + 2] / 65535.0f,
+					SrcColors[SrcIndex + 3] / 65535.0f
+				);
+				Color = DecodeFunction(Color);
+				DestColors[TexelIndex] = SaturateToHalfFloat(Color);
+			});
+			break;
 		}
-	}
-	break;
 
-	case ERawImageFormat::RGBA16F:
-	{
-		TArrayView64<const FFloat16Color> SrcColors = SrcImage.AsRGBA16F();
-		for (int64 TexelIndex = 0; TexelIndex < NumTexels; ++TexelIndex)
+		case ERawImageFormat::RGBA16F:
 		{
-			DestColors[TexelIndex] = SrcColors[TexelIndex].GetFloats();
-			DestColors[TexelIndex] = DecodeFunction(DestColors[TexelIndex]);
-			SaturateToHalfFloat(DestColors[TexelIndex]);
+			const FFloat16Color* SrcColors = SrcImage.AsRGBA16F().GetData();
+			ParallelLoop(TEXT("Texture.Linearize.PF"), NumJobs, TexelsPerJob, NumTexels, [DestColors, SrcColors, DecodeFunction](int64 TexelIndex)
+			{
+				FLinearColor Color = SrcColors[TexelIndex].GetFloats();
+				Color = DecodeFunction(Color);
+				DestColors[TexelIndex] = SaturateToHalfFloat(Color);
+			});
+			break;
 		}
-	}
-	break;
 
-	case ERawImageFormat::RGBA32F:
-	{
-		TArrayView64<const FLinearColor> SrcColors = SrcImage.AsRGBA32F();
-		for (int64 TexelIndex = 0; TexelIndex < NumTexels; ++TexelIndex)
+		case ERawImageFormat::RGBA32F:
 		{
-			DestColors[TexelIndex] = DecodeFunction(SrcColors[TexelIndex]);
-			SaturateToHalfFloat(DestColors[TexelIndex]);
+			const FLinearColor* SrcColors = SrcImage.AsRGBA32F().GetData();
+			ParallelLoop(TEXT("Texture.Linearize.PF"), NumJobs, TexelsPerJob, NumTexels, [DestColors, SrcColors, DecodeFunction](int64 TexelIndex)
+			{
+				FLinearColor Color = DecodeFunction(SrcColors[TexelIndex]);
+				DestColors[TexelIndex] = SaturateToHalfFloat(Color);
+			});
+			break;
 		}
-	}
-	break;
 
-	case ERawImageFormat::R16F:
-	{
-		TArrayView64<const FFloat16> SrcColors = SrcImage.AsR16F();
-		for (int64 TexelIndex = 0; TexelIndex < NumTexels; ++TexelIndex)
+		case ERawImageFormat::R16F:
 		{
-			DestColors[TexelIndex] = FLinearColor(SrcColors[TexelIndex].GetFloat(), 0, 0, 1);
-			DestColors[TexelIndex] = DecodeFunction(DestColors[TexelIndex]);
-			SaturateToHalfFloat(DestColors[TexelIndex]);
+			const FFloat16* SrcColors = SrcImage.AsR16F().GetData();
+			ParallelLoop(TEXT("Texture.Linearize.PF"), NumJobs, TexelsPerJob, NumTexels, [DestColors, SrcColors, DecodeFunction](int64 TexelIndex)
+			{
+				FLinearColor Color(SrcColors[TexelIndex].GetFloat(), 0, 0, 1);
+				Color = DecodeFunction(Color);
+				DestColors[TexelIndex] = SaturateToHalfFloat(Color);
+			});
+			break;
 		}
 	}
-	break;
-	}
 }
 
 int32 FImage::GetBytesPerPixel() const
