ImageCore NOP comment about strange and slow behavior in Linearize

diff --git a/Engine/Source/Runtime/ImageCore/Private/ImageCore.cpp b/Engine/Source/Runtime/ImageCore/Private/ImageCore.cpp
--- a/Engine/Source/Runtime/ImageCore/Private/ImageCore.cpp
+++ b/Engine/Source/Runtime/ImageCore/Private/ImageCore.cpp
@@ -539,22 +539,30 @@ void FImage::Linearize(uint8 SourceEncoding, FImage& DestImage) const
 	check(SrcImage.SizeX == DestImage.SizeX);
 	check(SrcImage.SizeY == DestImage.SizeY);
 	check(SrcImage.NumSlices == DestImage.NumSlices);
+	//NOTE: SrcImage has GammaSpace in addition to SourceEncoding
+	// @todo Oodle : the way this function behaves with SrcImage GammaSpace and SourceEncoding is not well defined
 
 	UE::Color::EEncoding SourceEncodingType = static_cast<UE::Color::EEncoding>(SourceEncoding);
 
 	if (SourceEncodingType == UE::Color::EEncoding::None)
 	{
+		// note SrcImage.GammaSpace is decoded by CopyImage, is that intended?
 		CopyImage(SrcImage, DestImage);
 		return;
 	}
 	else if (SourceEncodingType >= UE::Color::EEncoding::Max)
 	{
+		// note SrcImage.GammaSpace is decoded by CopyImage, is that intended?
 		UE_LOG(LogImageCore, Warning, TEXT("Invalid encoding %d, falling back to linearization using CopyImage."), SourceEncoding);
 		CopyImage(SrcImage, DestImage);
 		return;
 	}
 
-	const bool bDestIsGammaCorrected = DestImage.IsGammaCorrected();
+	// @todo Oodle : common case of sRGB decoding, just use CopyImage, it's much faster
+	// SourceEncodingType == UE::Color::EEncoding::sRGB
+	// this function calls a func pointer per pixel
+
+
 	const int64 NumTexels = int64(SrcImage.SizeX) * SrcImage.SizeY * SrcImage.NumSlices;
 
 	int64 TexelsPerJob;
