Fix editor crash in -nothreading mode

diff --git a/Engine/Source/Runtime/ImageCore/Private/ImageCore.cpp b/Engine/Source/Runtime/ImageCore/Private/ImageCore.cpp
--- a/Engine/Source/Runtime/ImageCore/Private/ImageCore.cpp
+++ b/Engine/Source/Runtime/ImageCore/Private/ImageCore.cpp
@@ -58,15 +58,14 @@ static inline int32 ParallelForComputeNumJobs(int64 & OutNumItemsPerJob,int64 Nu
 	// ParallelFor will actually make 6*NumWorkers batches and then make NumWorkers tasks that pop the batches
 	//	this helps with mismatched thread runtime
 	//	here we only make NumWorkers batches max
-	//	but this is rarely a problem in image cook because it is paralellized already at a the higher level
-
-	const int32 NumWorkers = FTaskGraphInterface::Get().GetNumWorkerThreads();
+	//	but this is rarely a problem in image cook because it is parallelized already at a the higher level
 
+	const int32 NumWorkers = FMath::Max(1, FTaskGraphInterface::Get().GetNumWorkerThreads());
 	int32 NumJobs = (int32)(NumItems / MinNumItemsPerJob); // round down
 	NumJobs = FMath::Clamp(NumJobs, int32(1), NumWorkers); 
 
 	OutNumItemsPerJob = (NumItems + NumJobs-1) / NumJobs; // round up
-	check( NumJobs*OutNumItemsPerJob >= NumItems ); 
+	check( NumJobs*OutNumItemsPerJob >= NumItems );
 
 	return NumJobs;
 }
