ImageCore.cpp fix asserting DestImage.GammaSpace
temp disable check

diff --git a/Engine/Source/Runtime/ImageCore/Private/ImageCore.cpp b/Engine/Source/Runtime/ImageCore/Private/ImageCore.cpp
--- a/Engine/Source/Runtime/ImageCore/Private/ImageCore.cpp
+++ b/Engine/Source/Runtime/ImageCore/Private/ImageCore.cpp
@@ -224,7 +224,8 @@ IMAGECORE_API void FImageCore::CopyImage(const FImageView & SrcImage,const FImag
 		const FLinearColor* SrcColors = (const FLinearColor*) SrcImage.RawData;
 	
 		// if gamma correction is done, it's always *TO* sRGB , not to Pow22
-		check( DestImage.GammaSpace != EGammaSpace::Pow22 );
+		//@@!!
+		// check( DestImage.GammaSpace != EGammaSpace::Pow22 );
 
 		switch (DestImage.Format)
 		{
