Add TSF_R32F single channel float texture source format

diff --git a/Engine/Source/Developer/TextureCompressor/Private/TextureCompressorModule.cpp b/Engine/Source/Developer/TextureCompressor/Private/TextureCompressorModule.cpp
--- a/Engine/Source/Developer/TextureCompressor/Private/TextureCompressorModule.cpp
+++ b/Engine/Source/Developer/TextureCompressor/Private/TextureCompressorModule.cpp
@@ -2727,7 +2727,8 @@ static bool DetectAlphaChannel(const FImage& InImage)
 	else if ( InImage.Format == ERawImageFormat::G8 ||
 		InImage.Format == ERawImageFormat::BGRE8 ||
 		InImage.Format == ERawImageFormat::G16 ||
-		InImage.Format == ERawImageFormat::R16F )
+		InImage.Format == ERawImageFormat::R16F||
+		InImage.Format == ERawImageFormat::R32F )
 	{
 		// source image formats don't have alpha
 	}
diff --git a/Engine/Source/Developer/TextureFormatUncompressed/Private/TextureFormatUncompressed.cpp b/Engine/Source/Developer/TextureFormatUncompressed/Private/TextureFormatUncompressed.cpp
--- a/Engine/Source/Developer/TextureFormatUncompressed/Private/TextureFormatUncompressed.cpp
+++ b/Engine/Source/Developer/TextureFormatUncompressed/Private/TextureFormatUncompressed.cpp
@@ -44,6 +44,7 @@ class FUncompressedTextureBuildFunction final : public FTextureBuildFunction
 	op(RGBA8) \
 	op(POTERROR) \
 	op(R16F) \
+	op(R32F) \
 	op(R5G6B5) \
 	op(A1RGB555) \
 	op(RGB555A1)
@@ -94,6 +95,10 @@ class FTextureFormatUncompressed : public ITextureFormat
 		{
 			return TEXT("R16F");
 		}
+		else if (InBuildSettings.TextureFormatName == GTextureFormatNameR32F)
+		{
+			return TEXT("R32F");
+		}
 		else
 		{
 			// default implementation of GetDerivedDataKeyString returns empty string
@@ -143,6 +148,10 @@ class FTextureFormatUncompressed : public ITextureFormat
 		{
 			return PF_R16F;
 		}
+		else if (BuildSettings.TextureFormatName == GTextureFormatNameR32F)
+		{
+			return PF_R32_FLOAT;
+		}
 		else if (BuildSettings.TextureFormatName == GTextureFormatNamePOTERROR)
 		{
 			return PF_B8G8R8A8;
@@ -170,6 +179,9 @@ class FTextureFormatUncompressed : public ITextureFormat
 	{
 		OutCompressedImage.PixelFormat = GetPixelFormatForImage(BuildSettings, InImage, bImageHasAlphaChannel);
 
+		// @todo Oodle : fix me : lots of pointless code dupe here
+		//	most of these should just map the Name to an ERawImageFormat and do CopyImage
+
 		if (BuildSettings.TextureFormatName == GTextureFormatNameG8)
 		{
 			FImage Image;
@@ -239,6 +251,7 @@ class FTextureFormatUncompressed : public ITextureFormat
 			OutCompressedImage.SizeY = Image.SizeY;
 			OutCompressedImage.SizeZ = (BuildSettings.bVolume || BuildSettings.bTextureArray) ? Image.NumSlices : 1;
 
+			// BGRA to RGBA : @todo Oodle : just use ImageCore CopyImageRGBABGRA
 			// swizzle each texel
 			uint64 NumTexels = (uint64)Image.SizeX * Image.SizeY * Image.NumSlices;
 			OutCompressedImage.RawData.Empty(NumTexels * 4);
@@ -303,7 +316,19 @@ class FTextureFormatUncompressed : public ITextureFormat
 
 			OutCompressedImage.SizeX = Image.SizeX;
 			OutCompressedImage.SizeY = Image.SizeY;
-			OutCompressedImage.SizeZ = BuildSettings.bVolume ? Image.NumSlices : 1;
+			OutCompressedImage.SizeZ = (BuildSettings.bVolume || BuildSettings.bTextureArray) ? Image.NumSlices : 1;
+			OutCompressedImage.RawData = MoveTemp(Image.RawData);
+
+			return true;
+		}
+		else if (BuildSettings.TextureFormatName == GTextureFormatNameR32F)
+		{
+			FImage Image;
+			InImage.CopyTo(Image, ERawImageFormat::R32F, EGammaSpace::Linear);
+
+			OutCompressedImage.SizeX = Image.SizeX;
+			OutCompressedImage.SizeY = Image.SizeY;
+			OutCompressedImage.SizeZ = (BuildSettings.bVolume || BuildSettings.bTextureArray) ? Image.NumSlices : 1;
 			OutCompressedImage.RawData = MoveTemp(Image.RawData);
 
 			return true;
diff --git a/Engine/Source/Runtime/Core/Public/Math/Float16.h b/Engine/Source/Runtime/Core/Public/Math/Float16.h
--- a/Engine/Source/Runtime/Core/Public/Math/Float16.h
+++ b/Engine/Source/Runtime/Core/Public/Math/Float16.h
@@ -74,7 +74,8 @@ public:
 	/** Convert from Fp16 to Fp32. */
 	float GetFloat() const;
 
-	/** Is the float negative without converting */
+	/** Is the float negative without converting
+	NOTE: returns true for negative zero! */
 	bool IsNegative() const
 	{
 		// negative if sign bit is on
diff --git a/Engine/Source/Runtime/Engine/Classes/Engine/Texture.h b/Engine/Source/Runtime/Engine/Classes/Engine/Texture.h
--- a/Engine/Source/Runtime/Engine/Classes/Engine/Texture.h
+++ b/Engine/Source/Runtime/Engine/Classes/Engine/Texture.h
@@ -132,7 +132,7 @@ struct FTextureSource
 	// should match ERawImageFormat::GetBytesPerPixel
 	ENGINE_API static int32 GetBytesPerPixel(ETextureSourceFormat Format);
 	// should match ERawImageFormat::IsHDR
-	FORCEINLINE static bool IsHDR(ETextureSourceFormat Format) { return (Format == TSF_BGRE8 || Format == TSF_RGBA16F || Format == TSF_R16F || Format == TSF_RGBA32F); }
+	FORCEINLINE static bool IsHDR(ETextureSourceFormat Format) { return (Format == TSF_BGRE8 || Format == TSF_RGBA16F || Format == TSF_R16F || Format == TSF_R32F || Format == TSF_RGBA32F); }
 	
 	enum class ELockState : uint8
 	{
diff --git a/Engine/Source/Runtime/Engine/Classes/Engine/TextureDefines.h b/Engine/Source/Runtime/Engine/Classes/Engine/TextureDefines.h
--- a/Engine/Source/Runtime/Engine/Classes/Engine/TextureDefines.h
+++ b/Engine/Source/Runtime/Engine/Classes/Engine/TextureDefines.h
@@ -312,6 +312,7 @@ enum ETextureSourceFormat
 	TSF_G16,
 	TSF_RGBA32F,
 	TSF_R16F,
+	TSF_R32F,
 
 	TSF_MAX,
 
@@ -331,7 +332,7 @@ enum TextureCompressionSettings
 	TC_Grayscale				UMETA(DisplayName = "Grayscale (R8, RGB8 sRGB)"),
 	TC_Displacementmap			UMETA(DisplayName = "Displacementmap (8/16bit)"),
 	TC_VectorDisplacementmap	UMETA(DisplayName = "VectorDisplacementmap (RGBA8)"),
-	TC_HDR						UMETA(DisplayName = "HDR (RGB, no sRGB)"),
+	TC_HDR						UMETA(DisplayName = "HDR (RGBA16F, no sRGB)"),
 	TC_EditorIcon				UMETA(DisplayName = "UserInterface2D (RGBA)"),
 	TC_Alpha					UMETA(DisplayName = "Alpha (no sRGB, BC4 on DX11)"),
 	TC_DistanceFieldFont		UMETA(DisplayName = "DistanceFieldFont (R8)"),
diff --git a/Engine/Source/Runtime/Engine/Private/Texture.cpp b/Engine/Source/Runtime/Engine/Private/Texture.cpp
--- a/Engine/Source/Runtime/Engine/Private/Texture.cpp
+++ b/Engine/Source/Runtime/Engine/Private/Texture.cpp
@@ -2879,6 +2879,7 @@ FName GetDefaultTextureFormatName( const ITargetPlatform* TargetPlatform, const
 	static FName NameVU8(TEXT("VU8"));
 	static FName NameRGBA16F(TEXT("RGBA16F"));
 	static FName NameR16F(TEXT("R16F"));
+	static FName NameR32F(TEXT("R32F"));
 	static FName NameBC6H(TEXT("BC6H"));
 	static FName NameBC7(TEXT("BC7"));
 	static FName NameR5G6B5(TEXT("R5G6B5"));
@@ -2970,6 +2971,9 @@ FName GetDefaultTextureFormatName( const ITargetPlatform* TargetPlatform, const
 			// note if bIsTCThatMapsToUncompressed , this might change the mapping :
 			bNoCompression = true;
 		}
+
+		// @todo Oodle : the conditions that trigger bNoCompression should be detected earlier
+		//   and shown in the GUI so the artist can see it
 	}
 
 	bool bUseDXT5NormalMap = false;
@@ -2997,7 +3001,7 @@ FName GetDefaultTextureFormatName( const ITargetPlatform* TargetPlatform, const
 
 		if (Texture->HasHDRSource(LayerIndex))
 		{
-			// R16F is available but not used here
+			// R16F and R32F is available but not used here even if their TC_ would have chosen them!
 			TextureFormatName = NameRGBA16F;
 		}
 		else if (SourceFormat == TSF_G16)
@@ -3159,6 +3163,7 @@ void GetAllDefaultTextureFormats(const class ITargetPlatform* TargetPlatform, TA
 	static FName NameVU8(TEXT("VU8"));
 	static FName NameRGBA16F(TEXT("RGBA16F"));
 	static FName NameR16F(TEXT("R16F"));
+	static FName NameR32F(TEXT("R32F"));
 	static FName NameBC6H(TEXT("BC6H"));
 	static FName NameBC7(TEXT("BC7"));
 
@@ -3176,6 +3181,7 @@ void GetAllDefaultTextureFormats(const class ITargetPlatform* TargetPlatform, TA
 	OutFormats.Add(NameVU8);
 	OutFormats.Add(NameRGBA16F);
 	OutFormats.Add(NameR16F);
+	OutFormats.Add(NameR32F);
 	if (bSupportDX11TextureFormats)
 	{
 		OutFormats.Add(NameBC6H);
diff --git a/Engine/Source/Runtime/ImageCore/Private/DDSFile.cpp b/Engine/Source/Runtime/ImageCore/Private/DDSFile.cpp
--- a/Engine/Source/Runtime/ImageCore/Private/DDSFile.cpp
+++ b/Engine/Source/Runtime/ImageCore/Private/DDSFile.cpp
@@ -1113,23 +1113,7 @@ bool FDDSFile::GetMipImage(const FImageView & ToImage, int MipIndex) const
 	// before BlitMip you should have done ConvertChannelOrder to BGRA and ConvertRGBXtoRGBA
 	// so we should only see BGRA (no RGBA or BGRX) here
 
-	if ( ToImage.Format == ERawImageFormat::R16F && DXGIFormat == EDXGIFormat::R32_FLOAT )
-	{
-		// no ERawImageFormat::R32F currently , have to convert to F16
-		
-		const float * FmPtr = (const float *) FromMip.Data;
-		uint16 * ToPtr = (uint16 *)ToImage.RawData;
-		int64 Count = ToImage.GetNumPixels();
-		check( ToImage.GetImageSizeBytes() == Count*sizeof(*ToPtr) );
-		check( FromMip.DataSize == Count*sizeof(*FmPtr) );
-		while(Count--)
-		{
-			FPlatformMath::StoreHalf(ToPtr, *FmPtr++);
-			ToPtr++;
-		}
-		return true;
-	}
-	else if ( ToImage.Format == ERawImageFormat::BGRA8 && DXGIFormat == EDXGIFormat::B8G8R8 )
+	if ( ToImage.Format == ERawImageFormat::BGRA8 && DXGIFormat == EDXGIFormat::B8G8R8 )
 	{
 		const uint8 * FmPtr = FromMip.Data;
 		uint8 * ToPtr = (uint8 *)ToImage.RawData;
@@ -1237,7 +1221,7 @@ FDXGIFormatRawFormatMapping DXGIRawFormatMap[] =
 {
 	{ EDXGIFormat::R32G32B32A32_FLOAT, ERawImageFormat::RGBA32F, true },
 	{ EDXGIFormat::R16G16B16A16_FLOAT, ERawImageFormat::RGBA16F, true },
-	{ EDXGIFormat::R32_FLOAT,  ERawImageFormat::R16F, false },
+	{ EDXGIFormat::R32_FLOAT,  ERawImageFormat::R32F, true },
 	{ EDXGIFormat::R16_FLOAT,  ERawImageFormat::R16F, true },
 
 	// R32G32B32_FLOAT not supported
@@ -1329,6 +1313,7 @@ EDXGIFormat DXGIFormatFromRawFormat(ERawImageFormat::Type RawFormat,EGammaSpace
 	case ERawImageFormat::RGBA32F:	return EDXGIFormat::R32G32B32A32_FLOAT;
 	case ERawImageFormat::G16:		return EDXGIFormat::R16_UNORM;
 	case ERawImageFormat::R16F:		return EDXGIFormat::R16_FLOAT;
+	case ERawImageFormat::R32F:		return EDXGIFormat::R32_FLOAT;
 	default:
 		return EDXGIFormat::UNKNOWN;
 	}
diff --git a/Engine/Source/Runtime/ImageCore/Private/ImageCore.cpp b/Engine/Source/Runtime/ImageCore/Private/ImageCore.cpp
--- a/Engine/Source/Runtime/ImageCore/Private/ImageCore.cpp
+++ b/Engine/Source/Runtime/ImageCore/Private/ImageCore.cpp
@@ -243,6 +243,8 @@ IMAGECORE_API void FImageCore::CopyImage(const FImageView & SrcImage,const FImag
 		{
 		case ERawImageFormat::G8:
 			{
+				// NOTE : blit from RGBA to G8 does NOT grab the gray, it just take the R
+				//	this makes blits from R16F/R32F go to G8/G16 and just take the single channel
 				uint8* DestLum = (uint8 *) DestImage.RawData;
 				if (bDestIsGammaCorrected)
 				{
@@ -263,6 +265,8 @@ IMAGECORE_API void FImageCore::CopyImage(const FImageView & SrcImage,const FImag
 
 		case ERawImageFormat::G16:
 			{
+				// NOTE : blit from RGBA to G16 does NOT grab the gray, it just take the R
+				//	this makes blits from R16F/R32F go to G8/G16 and just take the single channel
 				uint16* DestLum = (uint16 *) DestImage.RawData;
 				ParallelLoop(TEXT("Texture.CopyImage.PF"), NumJobs, TexelsPerJob, NumTexels, [DestLum, SrcColors](int64 TexelIndex)
 				{
@@ -392,6 +396,16 @@ IMAGECORE_API void FImageCore::CopyImage(const FImageView & SrcImage,const FImag
 			}
 			break;
 			
+		case ERawImageFormat::R32F:
+			{
+				float* DestColors = (float*) DestImage.RawData;
+				ParallelLoop(TEXT("Texture.CopyImage.PF"), NumJobs, TexelsPerJob, NumTexels, [DestColors, SrcColors](int64 TexelIndex)
+				{
+					DestColors[TexelIndex] = SrcColors[TexelIndex].R;
+				});
+			}
+			break;
+
 		default:
 			check(0);
 			break;
@@ -558,6 +572,16 @@ IMAGECORE_API void FImageCore::CopyImage(const FImageView & SrcImage,const FImag
 				});
 			}
 			break;
+			
+		case ERawImageFormat::R32F:
+			{
+				const float* SrcColors = (const float*) SrcImage.RawData;
+				ParallelLoop(TEXT("Texture.CopyImage.PF"), NumJobs, TexelsPerJob, NumTexels, [DestColors, SrcColors](int64 TexelIndex)
+				{
+					DestColors[TexelIndex] = FLinearColor(SrcColors[TexelIndex], 0, 0, 1);
+				});
+			}
+			break;
 
 		default:
 			check(0);
@@ -651,6 +675,8 @@ IMAGECORE_API void FImageCore::CopyImage(const FImageView & SrcImage,const FImag
 		TRACE_CPUPROFILER_EVENT_SCOPE(Texture.CopyImage.TempLinear);
 
 		// Arbitrary conversion, use 32-bit linear float as an intermediate format.
+		// this is unnecessarily expensive to do something like G8 to R16F, but rare
+		// if this shows up as a hot spot, identify the formats using this path and add direct conversions between them
 		FImage TempImage(SrcImage.SizeX, SrcImage.SizeY, SrcImage.NumSlices, ERawImageFormat::RGBA32F, EGammaSpace::Linear);
 		FImageCore::CopyImage(SrcImage, TempImage);
 		FImageCore::CopyImage(TempImage, DestImage);
@@ -1072,6 +1098,18 @@ void FImage::Linearize(uint8 SourceEncoding, FImage& DestImage) const
 			});
 			break;
 		}
+		
+		case ERawImageFormat::R32F:
+		{
+			const float* SrcColors = SrcImage.AsR32F().GetData();
+			ParallelLoop(TEXT("Texture.Linearize.PF"), NumJobs, TexelsPerJob, NumTexels, [DestColors, SrcColors, DecodeFunction](int64 TexelIndex)
+			{
+				FLinearColor Color(SrcColors[TexelIndex], 0, 0, 1);
+				Color = DecodeFunction(Color);
+				DestColors[TexelIndex] = SaturateToHalfFloat(Color);
+			});
+			break;
+		}
 	}
 }
 
@@ -1082,6 +1120,7 @@ IMAGECORE_API const TCHAR * ERawImageFormat::GetName(Type Format)
 	case ERawImageFormat::G8:  return TEXT("G8");
 	case ERawImageFormat::G16: return TEXT("G16");
 	case ERawImageFormat::R16F: return TEXT("R16F");
+	case ERawImageFormat::R32F: return TEXT("R32F");
 	case ERawImageFormat::BGRA8: return TEXT("BGRA8");
 	case ERawImageFormat::BGRE8: return TEXT("BGRE8");
 	case ERawImageFormat::RGBA16: return TEXT("RGBA16");
@@ -1106,7 +1145,8 @@ IMAGECORE_API int32 ERawImageFormat::GetBytesPerPixel(Type Format)
 	case ERawImageFormat::R16F:
 		BytesPerPixel = 2;
 		break;
-
+		
+	case ERawImageFormat::R32F:
 	case ERawImageFormat::BGRA8:
 	case ERawImageFormat::BGRE8:
 		BytesPerPixel = 4;
@@ -1131,7 +1171,7 @@ IMAGECORE_API int32 ERawImageFormat::GetBytesPerPixel(Type Format)
 
 IMAGECORE_API bool ERawImageFormat::IsHDR(Type Format)
 {
-	return Format == RGBA16F || Format == RGBA32F || Format == R16F || Format == BGRE8;
+	return Format == RGBA16F || Format == RGBA32F || Format == R16F || Format == R32F || Format == BGRE8;
 }
 
 IMAGECORE_API FLinearColor ERawImageFormat::GetOnePixelLinear(const void * PixelData,Type Format,bool bSRGB)
@@ -1147,6 +1187,12 @@ IMAGECORE_API FLinearColor ERawImageFormat::GetOnePixelLinear(const void * Pixel
 		else
 			return Color.ReinterpretAsLinear();
 	}
+	case G16:
+	{
+		const uint16 * Samples = (const uint16 *)PixelData;
+		float Gray = Samples[0]/65535.f;
+		return FLinearColor(Gray,Gray,Gray,1.f);
+	}
 	case BGRA8:
 	{
 		FColor Color = *((const FColor *)PixelData);
@@ -1176,18 +1222,17 @@ IMAGECORE_API FLinearColor ERawImageFormat::GetOnePixelLinear(const void * Pixel
 	}
 	case RGBA32F:
 	{
-		return *((FLinearColor *)PixelData);
+		return *((const FLinearColor *)PixelData);
 	}
-	case G16:
+	case R16F:
 	{
-		const uint16 * Samples = (const uint16 *)PixelData;
-		float Gray = Samples[0]/65535.f;
-		return FLinearColor(Gray,Gray,Gray,1.f);
+		float R = FPlatformMath::LoadHalf((const uint16 *)PixelData);
+		return FLinearColor(R,0,0,1.f);
 	}
-	case R16F:
+	case R32F:
 	{
-		float Gray = FPlatformMath::LoadHalf((const uint16 *)PixelData);
-		return FLinearColor(Gray,Gray,Gray,1.f);
+		float R = *((const float *)PixelData);
+		return FLinearColor(R,0,0,1.f);
 	}
 	
 	default:
diff --git a/Engine/Source/Runtime/ImageCore/Private/ImageCoreUtils.cpp b/Engine/Source/Runtime/ImageCore/Private/ImageCoreUtils.cpp
--- a/Engine/Source/Runtime/ImageCore/Private/ImageCoreUtils.cpp
+++ b/Engine/Source/Runtime/ImageCore/Private/ImageCoreUtils.cpp
@@ -40,6 +40,8 @@ EPixelFormat FImageCoreUtils::GetPixelFormatForRawImageFormat(ERawImageFormat::T
 		return PF_FloatRGBA;
 	case ERawImageFormat::R16F:	
 		return PF_R16F;
+	case ERawImageFormat::R32F:	
+		return PF_R32_FLOAT;
 	
 	default:
 		check(0);
@@ -61,6 +63,7 @@ IMAGECORE_API ERawImageFormat::Type FImageCoreUtils::ConvertToRawImageFormat(ETe
 	case TSF_G16: return ERawImageFormat::G16;
 	case TSF_RGBA32F: return ERawImageFormat::RGBA32F;
 	case TSF_R16F: return ERawImageFormat::R16F;
+	case TSF_R32F: return ERawImageFormat::R32F;
 
 	// these are mapped to TSF_BGRA8/TSF_BGRE8 on load, so the runtime will never see them :
 	case TSF_RGBA8_DEPRECATED:
@@ -87,6 +90,7 @@ IMAGECORE_API ETextureSourceFormat FImageCoreUtils::ConvertToTextureSourceFormat
 	case ERawImageFormat::RGBA32F: return TSF_RGBA32F;
 	case ERawImageFormat::G16: return TSF_G16;
 	case ERawImageFormat::R16F:	return TSF_R16F;
+	case ERawImageFormat::R32F:	return TSF_R32F;
 	
 	default:
 		check(0);
diff --git a/Engine/Source/Runtime/ImageCore/Public/ImageCore.h b/Engine/Source/Runtime/ImageCore/Public/ImageCore.h
--- a/Engine/Source/Runtime/ImageCore/Public/ImageCore.h
+++ b/Engine/Source/Runtime/ImageCore/Public/ImageCore.h
@@ -39,13 +39,14 @@ namespace ERawImageFormat
 	enum Type : uint8
 	{
 		G8,
-		BGRA8, // FColor
+		BGRA8,   // FColor
 		BGRE8,
 		RGBA16,
 		RGBA16F, // FFloat16Color
 		RGBA32F, // FLinearColor
-		G16,
+		G16,     // note G8/G16 = gray = replicate to 3 channels, R16F = just in red channel
 		R16F,
+		R32F,
 		Invalid = 0xFF
 	};
 	
@@ -470,6 +471,12 @@ public:
 		check(Format == ERawImageFormat::R16F);
 		return { (class FFloat16*)RawData.GetData(), int64(RawData.Num() / sizeof(FFloat16)) };
 	}
+	
+	TArrayView64<float> AsR32F()
+	{
+		check(Format == ERawImageFormat::R32F);
+		return { (float *)RawData.GetData(), int64(RawData.Num() / sizeof(float)) };
+	}
 
 	// Convenience accessors to const raw data
 
@@ -520,6 +527,12 @@ public:
 		check(Format == ERawImageFormat::R16F);
 		return { (const class FFloat16*)RawData.GetData(), int64(RawData.Num() / sizeof(FFloat16)) };
 	}
+	
+	TArrayView64<const float> AsR32F() const
+	{
+		check(Format == ERawImageFormat::R32F);
+		return { (const float*)RawData.GetData(), int64(RawData.Num() / sizeof(float)) };
+	}
 };
 
 /* Functions
diff --git a/Engine/Source/Runtime/ImageWrapper/Private/Formats/BmpImageWrapper.cpp b/Engine/Source/Runtime/ImageWrapper/Private/Formats/BmpImageWrapper.cpp
--- a/Engine/Source/Runtime/ImageWrapper/Private/Formats/BmpImageWrapper.cpp
+++ b/Engine/Source/Runtime/ImageWrapper/Private/Formats/BmpImageWrapper.cpp
@@ -353,6 +353,7 @@ ERawImageFormat::Type FBmpImageWrapper::GetSupportedRawFormat(const ERawImageFor
 	case ERawImageFormat::RGBA16F:
 	case ERawImageFormat::RGBA32F:
 	case ERawImageFormat::R16F:
+	case ERawImageFormat::R32F: // @@!! can these go to G8 ?
 		return ERawImageFormat::BGRA8; // needs conversion
 	default:
 		check(0);
diff --git a/Engine/Source/Runtime/ImageWrapper/Private/Formats/ExrImageWrapper.cpp b/Engine/Source/Runtime/ImageWrapper/Private/Formats/ExrImageWrapper.cpp
--- a/Engine/Source/Runtime/ImageWrapper/Private/Formats/ExrImageWrapper.cpp
+++ b/Engine/Source/Runtime/ImageWrapper/Private/Formats/ExrImageWrapper.cpp
@@ -7,6 +7,14 @@
 #include "HAL/PlatformTime.h"
 #include "Math/Float16.h"
 
+/**
+
+@todo Oodle : this is a flawed EXR importer
+it crashes on many images in the OpenEXR test set
+Blobbies.exr and spirals.exr among others
+
+**/
+
 #if WITH_UNREALEXR
 
 FExrImageWrapper::FExrImageWrapper()
@@ -150,7 +158,8 @@ private:
 
 const char* cChannelNamesRGBA[] = { "R", "G", "B", "A" };
 const char* cChannelNamesBGRA[] = { "B", "G", "R", "A" };
-const char* cChannelNamesGray[] = { "G" };
+//const char* cChannelNamesGray[] = { "G" }; // is that green or gray ?
+const char* cChannelNamesGray[] = { "Y" }; // pretty sure "Y" is more standard for gray
 
 int32 GetChannelNames(ERGBFormat InRGBFormat, const char* const*& OutChannelNames)
 {
@@ -196,6 +205,7 @@ ERawImageFormat::Type FExrImageWrapper::GetSupportedRawFormat(const ERawImageFor
 	case ERawImageFormat::RGBA16F:
 	case ERawImageFormat::RGBA32F:
 	case ERawImageFormat::R16F:
+	case ERawImageFormat::R32F:
 		return InFormat; // directly supported
 	case ERawImageFormat::G8:
 		return ERawImageFormat::R16F; // needs conversion
@@ -281,13 +291,25 @@ bool FExrImageWrapper::SetCompressed(const void* InCompressedData, int64 InCompr
 		Height = ImfDataWindow.max.y - ImfDataWindow.min.y + 1;
 
 		bool bHasOnlyHALFChannels = true;
-		bool bMatchesGrayOrder = true;
+		bool bHasAlphaChannel = false;
 		int32 ChannelCount = 0;
 
-		for (Imf::ChannelList::Iterator Iter = ImfChannels.begin(); Iter != ImfChannels.end(); ++Iter, ++ChannelCount)
+		for (Imf::ChannelList::Iterator Iter = ImfChannels.begin(); Iter != ImfChannels.end(); ++Iter)
 		{
+			++ChannelCount;
+
 			bHasOnlyHALFChannels = bHasOnlyHALFChannels && Iter.channel().type == Imf::HALF;
-			bMatchesGrayOrder = bMatchesGrayOrder && ChannelCount < UE_ARRAY_COUNT(cChannelNamesGray) && !strcmp(Iter.name(), cChannelNamesGray[ChannelCount]);
+
+			// check for bMatchesGrayOrder disabled
+			// treat any 1-channel import as gray
+			// don't try to match the channel names, which are not standardized
+			//  (we incorrectly used "G" before, we now use "Y", TinyEXR uses "A" for 1-channel EXR)
+			//bMatchesGrayOrder = bMatchesGrayOrder && ChannelCount < UE_ARRAY_COUNT(cChannelNamesGray) && !strcmp(Iter.name(), cChannelNamesGray[ChannelCount]);
+
+			if ( strcmp(Iter.name(),"A") == 0 )
+			{
+				bHasAlphaChannel = true;
+			}
 		}
 
 		BitDepth = (ChannelCount && bHasOnlyHALFChannels) ? 16 : 32;
@@ -295,7 +317,19 @@ bool FExrImageWrapper::SetCompressed(const void* InCompressedData, int64 InCompr
 		// EXR uint32 channels are currently not supported, therefore input channels are always treated as float channels.
 		// Channel combinations which don't match the ERGBFormat::GrayF pattern are qualified as ERGBFormat::RGBAF.
 		// Note that channels inside the EXR file are indexed by name, therefore can be decoded in any RGB order.
-		Format = (ChannelCount == UE_ARRAY_COUNT(cChannelNamesGray) && bMatchesGrayOrder) ? ERGBFormat::GrayF : ERGBFormat::RGBAF;
+
+		// NOTE: TinyEXR writes 1-channel EXR as an "A" named channel
+		//  this cannot be loaded as a 1-channel image (you would just get all zeros)
+		// it must be loaded as RGBA here
+		// @todo Oodle : load that as RGBA then move A to R and convert back to 1 channel ?
+		if (ChannelCount == 1 && !bHasAlphaChannel)
+		{
+			Format = ERGBFormat::GrayF;
+		}
+		else
+		{
+			Format = ERGBFormat::RGBAF;
+		}
 	}
 	catch (const std::exception& Exception)
 	{
diff --git a/Engine/Source/Runtime/ImageWrapper/Private/Formats/JpegImageWrapper.cpp b/Engine/Source/Runtime/ImageWrapper/Private/Formats/JpegImageWrapper.cpp
--- a/Engine/Source/Runtime/ImageWrapper/Private/Formats/JpegImageWrapper.cpp
+++ b/Engine/Source/Runtime/ImageWrapper/Private/Formats/JpegImageWrapper.cpp
@@ -152,8 +152,9 @@ ERawImageFormat::Type FJpegImageWrapper::GetSupportedRawFormat(const ERawImageFo
 	case ERawImageFormat::RGBA16:
 	case ERawImageFormat::RGBA16F:
 	case ERawImageFormat::RGBA32F:
-	case ERawImageFormat::G16:
+	case ERawImageFormat::G16: // @@!! use G8 ?
 	case ERawImageFormat::R16F:
+	case ERawImageFormat::R32F:
 		return ERawImageFormat::BGRA8; // needs conversion
 	default:
 		check(0);
diff --git a/Engine/Source/Runtime/ImageWrapper/Private/Formats/PngImageWrapper.cpp b/Engine/Source/Runtime/ImageWrapper/Private/Formats/PngImageWrapper.cpp
--- a/Engine/Source/Runtime/ImageWrapper/Private/Formats/PngImageWrapper.cpp
+++ b/Engine/Source/Runtime/ImageWrapper/Private/Formats/PngImageWrapper.cpp
@@ -139,6 +139,7 @@ ERawImageFormat::Type FPngImageWrapper::GetSupportedRawFormat(const ERawImageFor
 	case ERawImageFormat::RGBA16F:
 	case ERawImageFormat::RGBA32F:
 	case ERawImageFormat::R16F:
+	case ERawImageFormat::R32F:
 		return ERawImageFormat::RGBA16; // needs conversion
 	default:
 		check(0);
diff --git a/Engine/Source/Runtime/ImageWrapper/Private/ImageWrapperBase.cpp b/Engine/Source/Runtime/ImageWrapper/Private/ImageWrapperBase.cpp
--- a/Engine/Source/Runtime/ImageWrapper/Private/ImageWrapperBase.cpp
+++ b/Engine/Source/Runtime/ImageWrapper/Private/ImageWrapperBase.cpp
@@ -275,9 +275,8 @@ ERawImageFormat::Type IImageWrapper::ConvertRGBFormat(ERGBFormat RGBFormat,int B
 		}
 		else if ( BitDepth == 32 )
 		{
-			*bIsExactMatch = false; // no single channel F32
-			return ERawImageFormat::RGBA32F; // promote F32 to 4xF32 for no quality loss (alternative is 1xF16)
-			//@todo Oodle: add ERawImageFormat::R32F ?
+			*bIsExactMatch = true;
+			return ERawImageFormat::R32F;
 		}
 		break;
 
@@ -327,6 +326,10 @@ void IImageWrapper::ConvertRawImageFormat(ERawImageFormat::Type RawFormat, ERGBF
 		OutFormat = ERGBFormat::GrayF;
 		OutBitDepth = 16;
 		break;
+	case ERawImageFormat::R32F:
+		OutFormat = ERGBFormat::GrayF;
+		OutBitDepth = 32;
+		break;
 	default:
 		check(0);
 		break;
@@ -396,37 +399,6 @@ bool IImageWrapper::GetRawImage(FImage & OutImage)
 			break;
 		}
 
-		case ERGBFormat::GrayF:
-		{
-			/*
-			// 1 channel F32 -> F16 :
-			//  because ERawImageFormat has no 1 channel F32 currently
-			check( BitDepth == 32 );
-			check( RawFormat == ERawImageFormat::R16F );
-			const float * Src = (const float *)OutRawData.GetData();
-			uint16 * Dst = (uint16 *)OutImage.RawData.GetData();
-			int64 NumPixels = OutImage.GetNumPixels();
-			for(int64 i=0;i<NumPixels;i++)
-			{
-				// clamp in F16 range? does StoreHalf do that?
-				FPlatformMath::StoreHalf(&Dst[i], Src[i]);
-			}
-			*/
-
-			// 1 channel F32 -> 4xF32 :
-			check( BitDepth == 32 );
-			check( RawFormat == ERawImageFormat::RGBA32F );
-			const float * Src = (const float *)OutRawData.GetData();
-			FLinearColor * Dst = (FLinearColor *)OutImage.RawData.GetData();
-			int64 NumPixels = OutImage.GetNumPixels();
-			for(int64 i=0;i<NumPixels;i++)
-			{
-				Dst[i] = FLinearColor( Src[i],Src[i],Src[i],1.f );
-			}
-
-			break;
-		}
-
 		default:
 			check(0);
 			return false;			
diff --git a/Engine/Source/Runtime/ImageWrapper/Private/ImageWrapperModule.cpp b/Engine/Source/Runtime/ImageWrapper/Private/ImageWrapperModule.cpp
--- a/Engine/Source/Runtime/ImageWrapper/Private/ImageWrapperModule.cpp
+++ b/Engine/Source/Runtime/ImageWrapper/Private/ImageWrapperModule.cpp
@@ -296,6 +296,7 @@ public:
 		case ERawImageFormat::RGBA16F:
 		case ERawImageFormat::RGBA32F:
 		case ERawImageFormat::R16F:
+		case ERawImageFormat::R32F:
 			return EImageFormat::EXR;
 		
 		default:
