fix VT Image Resize with Pow22 Gamma
Resize converts Pow22 to sRGB because it can not output Pow22

diff --git a/Engine/Source/Developer/TextureBuildUtilities/Private/TextureBuildUtilities.cpp b/Engine/Source/Developer/TextureBuildUtilities/Private/TextureBuildUtilities.cpp
--- a/Engine/Source/Developer/TextureBuildUtilities/Private/TextureBuildUtilities.cpp
+++ b/Engine/Source/Developer/TextureBuildUtilities/Private/TextureBuildUtilities.cpp
@@ -68,11 +68,12 @@ TEXTUREBUILDUTILITIES_API const FName TextureFormatRemovePrefixFromName(FName co
 
 TEXTUREBUILDUTILITIES_API ERawImageFormat::Type GetVirtualTextureBuildIntermediateFormat(const FTextureBuildSettings& BuildSettings)
 {
-	FName TextureFormatPrefix;
-	const FName TextureFormatName = TextureFormatRemovePrefixFromName(BuildSettings.TextureFormatName, TextureFormatPrefix);
+	const FName TextureFormatName = TextureFormatRemovePrefixFromName(BuildSettings.TextureFormatName);
+
+	// note: using RGBA16F when the Source is HDR but the output is not HDR is not needed
+	//	you could use BGRA8 intermediate in that case
+	//	but it's rare and not a big problem, so leave it alone for now
 
-	// @todo Oodle: @@!! bIsHdr shouldn't be looking at HDRSource !
-	//	FIX ME beware possible changes to output; have to verify it's a nop on real data
 	const bool bIsHdr = BuildSettings.bHDRSource || TextureFormatIsHdr(TextureFormatName);
 
 	if (bIsHdr)
diff --git a/Engine/Source/Runtime/Engine/Private/TextureDerivedData.cpp b/Engine/Source/Runtime/Engine/Private/TextureDerivedData.cpp
--- a/Engine/Source/Runtime/Engine/Private/TextureDerivedData.cpp
+++ b/Engine/Source/Runtime/Engine/Private/TextureDerivedData.cpp
@@ -70,6 +70,7 @@
 // This is put in the DDC1 key but NOT in the DDC2 key
 #define TEXTURE_VT_DERIVEDDATA_VER	TEXT("7C16439390E24F1F9468894FB4D4BC54")
 
+
 #if ENABLE_COOK_STATS
 namespace TextureCookStats
 {
@@ -194,6 +195,15 @@ static void SerializeForKey(FArchive& Ar, const FTextureBuildSettings& Settings)
 		TempByte = 0; Ar << TempByte;
 		TempByte = Settings.LossyCompressionAmount; Ar << TempByte; // Lossy compression currently only used by VT
 		TempByte = Settings.bApplyYCoCgBlockScale; Ar << TempByte; // YCoCg currently only used by VT
+
+		
+		// @todo SerializeForKey these can go away whenever we bump the overall ddc key:
+		if ( Settings.bSRGB && Settings.bUseLegacyGamma )
+		{
+			// processing changed, modify ddc key :
+			TempGuid = FGuid(0xA227BEFC,0x9F8643C6,0x81580369,0xC4C6F73E);
+			Ar << TempGuid;
+		}
 	}
 
 	// Avoid changing key if texture is not being downscaled
diff --git a/Engine/Source/Runtime/Engine/Private/VT/VirtualTextureDataBuilder.cpp b/Engine/Source/Runtime/Engine/Private/VT/VirtualTextureDataBuilder.cpp
--- a/Engine/Source/Runtime/Engine/Private/VT/VirtualTextureDataBuilder.cpp
+++ b/Engine/Source/Runtime/Engine/Private/VT/VirtualTextureDataBuilder.cpp
@@ -937,15 +937,22 @@ void FVirtualTextureDataBuilder::BuildSourcePixels(const FTextureSourceData& Sou
 				for (const FImage& SrcMip : SourceMips)
 				{
 					FImage* ScaledMip = new(ScaledSourceMips) FImage;
-					SrcMip.ResizeTo(*ScaledMip, SrcMip.SizeX * LocalBlockSizeScale, SrcMip.SizeY * LocalBlockSizeScale, SrcMip.Format, SrcMip.GammaSpace);
+					// Pow22 cannot be used as a destination gamma, so change it to sRGB now :
+					EGammaSpace GammaSpace = (SrcMip.GammaSpace == EGammaSpace::Pow22) ? EGammaSpace::sRGB : SrcMip.GammaSpace;
+					SrcMip.ResizeTo(*ScaledMip, SrcMip.SizeX * LocalBlockSizeScale, SrcMip.SizeY * LocalBlockSizeScale, SrcMip.Format, GammaSpace);
 				}
 
 				for (const FImage& SrcMip : *CompositeSourceMips)
 				{
 					FImage* ScaledMip = new(ScaledCompositeMips) FImage;
-					SrcMip.ResizeTo(*ScaledMip, SrcMip.SizeX * LocalBlockSizeScale, SrcMip.SizeY * LocalBlockSizeScale, SrcMip.Format, SrcMip.GammaSpace);
+					// Pow22 cannot be used as a destination gamma, so change it to sRGB now :
+					EGammaSpace GammaSpace = (SrcMip.GammaSpace == EGammaSpace::Pow22) ? EGammaSpace::sRGB : SrcMip.GammaSpace;
+					SrcMip.ResizeTo(*ScaledMip, SrcMip.SizeX * LocalBlockSizeScale, SrcMip.SizeY * LocalBlockSizeScale, SrcMip.Format, GammaSpace);
 				}
 
+				// Pow22 was converted to sRGB by Resize :
+				TBSettings.bUseLegacyGamma = false;
+
 				uint32 NumMipsInTail, ExtData;
 				bBuildTextureResult = Compressor->BuildTexture(ScaledSourceMips, ScaledCompositeMips, TBSettings, CurDebugTexturePathName, CompressedMips, NumMipsInTail, ExtData);
 			}
diff --git a/Engine/Source/Runtime/ImageCore/Private/ImageCore.cpp b/Engine/Source/Runtime/ImageCore/Private/ImageCore.cpp
--- a/Engine/Source/Runtime/ImageCore/Private/ImageCore.cpp
+++ b/Engine/Source/Runtime/ImageCore/Private/ImageCore.cpp
@@ -752,6 +752,7 @@ static FLinearColor SampleImage(const FLinearColor* Pixels, int Width, int Heigh
 
 static void ResizeImage(const FImage& SrcImage, FImage& DestImage)
 {
+	// Src and Dest should both now be RGBA32F and Linear gamma
 	const FLinearColor* SrcPixels = SrcImage.AsRGBA32F().GetData();
 	FLinearColor* DestPixels = DestImage.AsRGBA32F().GetData();
 	const float DestToSrcScaleX = (float)SrcImage.SizeX / (float)DestImage.SizeX;
@@ -871,9 +872,20 @@ void FImage::ResizeTo(FImage& DestImage, int32 DestSizeX, int32 DestSizeY, ERawI
 		CopyTo(TempSrcImage, ERawImageFormat::RGBA32F, EGammaSpace::Linear);
 		SrcImagePtr = &TempSrcImage;
 	}
+	else if ( GammaSpace != EGammaSpace::Linear )
+	{
+		UE_LOG(LogImageCore, Warning, TEXT("Resize from source Format RGBA32F was called but source GammaSpace is not Linear"));
+	}
+
+	// now SrcImagePtr should be RGBA32F and Linear
 
 	if (DestFormat == ERawImageFormat::RGBA32F)
 	{
+		if ( DestGammaSpace != EGammaSpace::Linear )
+		{
+			UE_LOG(LogImageCore, Warning, TEXT("Resize to DestFormat RGBA32F was called but DestGammaSpace is not Linear"));
+		}
+
 		DestImage.SizeX = DestSizeX;
 		DestImage.SizeY = DestSizeY;
 		DestImage.NumSlices = 1;
@@ -884,14 +896,21 @@ void FImage::ResizeTo(FImage& DestImage, int32 DestSizeX, int32 DestSizeY, ERawI
 	}
 	else
 	{
+		// first resize from RGBA32F to RGBA32F
 		FImage TempDestImage;
 		TempDestImage.SizeX = DestSizeX;
 		TempDestImage.SizeY = DestSizeY;
 		TempDestImage.NumSlices = 1;
 		TempDestImage.Format = ERawImageFormat::RGBA32F;
-		TempDestImage.GammaSpace = DestGammaSpace;
+		TempDestImage.GammaSpace = EGammaSpace::Linear;
 		InitImageStorage(TempDestImage);
 		ResizeImage(*SrcImagePtr, TempDestImage);
+
+		// then convert to dest format/gamma :
+		if ( DestGammaSpace == EGammaSpace::Pow22 )
+		{
+			UE_LOG(LogImageCore, Warning, TEXT("Resize incorrectly used with Pow22 Dest Gamma"));
+		}
 		TempDestImage.CopyTo(DestImage, DestFormat, DestGammaSpace);
 	}
 }
