Clean up some terminology

#rnx

diff --git a/Engine/Source/Runtime/ImageCore/Private/DDSFile.cpp b/Engine/Source/Runtime/ImageCore/Private/DDSFile.cpp
--- a/Engine/Source/Runtime/ImageCore/Private/DDSFile.cpp
+++ b/Engine/Source/Runtime/ImageCore/Private/DDSFile.cpp
@@ -297,7 +297,7 @@ static const FFOURCCToDXGI FOURCCToDXGITable[] =
 
 static const FDXGIFormatInfo* DXGIFormatGetInfo(EDXGIFormat InFormat)
 {
-	// need to handle this special because UNKNOWN _does_ appear in the master list
+	// need to handle this special because UNKNOWN _does_ appear in the SupportedFormatList
 	// but we don't want to treat it as legal
 	if (InFormat == EDXGIFormat::UNKNOWN)
 	{
