Fix 6-image array DDS files being imported as cubes

diff --git a/Engine/Source/Runtime/ImageCore/Private/DDSFile.cpp b/Engine/Source/Runtime/ImageCore/Private/DDSFile.cpp
--- a/Engine/Source/Runtime/ImageCore/Private/DDSFile.cpp
+++ b/Engine/Source/Runtime/ImageCore/Private/DDSFile.cpp
@@ -448,13 +448,19 @@ EDDSError FDDSFile::Validate() const
 	// Cubemaps need to be square and have a valid array count
 	if (CreateFlags & CREATE_FLAG_CUBEMAP)
 	{
+		if (Dimension != 2 )
+		{
+			UE_LOG(LogDDSFile, Warning, TEXT("Cubemap must be dimension 2!"));
+			return EDDSError::BadCubemap;
+		}
+
 		if (Width != Height || Depth != 1)
 		{
 			UE_LOG(LogDDSFile, Warning, TEXT("Cubemap has non-square faces or non-1 depth!"));
 			return EDDSError::BadCubemap;
 		}
 
-		if ((ArraySize % 6) != 0)
+		if (ArraySize < 6 || (ArraySize % 6) != 0)
 		{
 			UE_LOG(LogDDSFile, Warning, TEXT("Cubemap or cubemap array doesn't have a multiple of 6 faces."));
 			return EDDSError::BadCubemap;
@@ -489,7 +495,7 @@ bool FDDSFile::IsValidTextureCube() const
 		return false;
 	}
 	
-	if ( Dimension == 2 && ArraySize >= 6 && (ArraySize%6) == 0 )
+	if (CreateFlags & CREATE_FLAG_CUBEMAP)
 	{
 		return true;
 	}
@@ -503,7 +509,13 @@ bool FDDSFile::IsValidTextureArray() const
 	{
 		return false;
 	}
-	
+
+	if (CreateFlags & CREATE_FLAG_CUBEMAP)
+	{
+		// don't identify cubes as arrays
+		return false;
+	}
+
 	if ( Dimension == 1 || Dimension == 2 )
 	{
 		if ( ArraySize > 1 )
