Fix incorrect check in SampleImage.

diff --git a/Engine/Source/Runtime/ImageCore/Private/ImageCore.cpp b/Engine/Source/Runtime/ImageCore/Private/ImageCore.cpp
--- a/Engine/Source/Runtime/ImageCore/Private/ImageCore.cpp
+++ b/Engine/Source/Runtime/ImageCore/Private/ImageCore.cpp
@@ -733,7 +733,7 @@ static FLinearColor SampleImage(const FLinearColor* Pixels, int Width, int Heigh
 	const int64 TexelX1 = FMath::Min<int64>(TexelX0 + 1, Width - 1);
 	const int64 TexelY1 = FMath::Min<int64>(TexelY0 + 1, Height - 1);
 	checkSlow(TexelX0 >= 0 && TexelX0 < Width);
-	checkSlow(TexelY0 >= 0 && TexelY0 < Width);
+	checkSlow(TexelY0 >= 0 && TexelY0 < Height);
 
 	const float FracX1 = FMath::Frac(X);
 	const float FracY1 = FMath::Frac(Y);
