fix ThumbnailManager was saving jpegs with .png extension
move some functionality from FImage to FImageView

diff --git a/Engine/Source/Editor/UnrealEd/Private/ThumbnailManager.cpp b/Engine/Source/Editor/UnrealEd/Private/ThumbnailManager.cpp
--- a/Engine/Source/Editor/UnrealEd/Private/ThumbnailManager.cpp
+++ b/Engine/Source/Editor/UnrealEd/Private/ThumbnailManager.cpp
@@ -289,7 +289,7 @@ bool UThumbnailManager::CaptureProjectThumbnail(FViewport* Viewport, const FStri
 		int32 ScaledSize  = FMath::Min<uint32>(AutoScreenshotSize, CropSize);
 
 		//calculations for cropping
-		TArray<FColor> CroppedBitmap;
+		TArray64<FColor> CroppedBitmap;
 		CroppedBitmap.AddUninitialized(CropSize*CropSize);
 
 		//Crop the image
@@ -305,21 +305,31 @@ bool UThumbnailManager::CaptureProjectThumbnail(FViewport* Viewport, const FStri
 			FMemory::Memcpy(DstPtr, SrcPtr, CropSize * 4);
 		}
 
+		FImageView CroppedImage(CroppedBitmap.GetData(),CropSize,CropSize);
+		//Viewport ReadPixels seems to have A = 0, make sure it is set to opaque for image save
+		FImageCore::SetAlphaOpaque(CroppedImage);
+
 		//Scale image down if needed
-		TArray<FColor> ScaledBitmap;
+		FImage ScaledImage;
+		FImageView SaveImage;
 		if (ScaledSize < CropSize)
 		{
-			FImageUtils::ImageResize( CropSize, CropSize, CroppedBitmap, ScaledSize, ScaledSize, ScaledBitmap, true );
+			FImageCore::ResizeTo(CroppedImage,ScaledImage,ScaledSize,ScaledSize,ERawImageFormat::BGRA8,EGammaSpace::sRGB);
+			SaveImage = ScaledImage;
 		}
 		else
 		{
 			//just copy the data over. sizes are the same
-			ScaledBitmap = CroppedBitmap;
+			SaveImage = CroppedImage;
 		}
 
 		// Compress the scaled image
-		TArray<uint8> ScaledPng;
-		FImageUtils::ThumbnailCompressImageArray(ScaledSize, ScaledSize, ScaledBitmap, ScaledPng);
+		// OutputFilename is a .png in current use
+		TArray64<uint8> ScaledPng;
+		if ( ! FImageUtils::CompressImage(ScaledPng, *OutputFilename, SaveImage) )
+		{
+			return false;
+		}
 
 		// Save to file
 		const FString ScreenShotPath = FPaths::GetPath(OutputFilename);
@@ -359,6 +369,10 @@ bool UThumbnailManager::CaptureProjectThumbnail(FViewport* Viewport, const FStri
 				return true;
 			}
 		}
+		else
+		{
+			// failed to make output dir?
+		}
 	}
 
 	return false;
diff --git a/Engine/Source/Runtime/ImageCore/Private/ImageCore.cpp b/Engine/Source/Runtime/ImageCore/Private/ImageCore.cpp
--- a/Engine/Source/Runtime/ImageCore/Private/ImageCore.cpp
+++ b/Engine/Source/Runtime/ImageCore/Private/ImageCore.cpp
@@ -750,14 +750,17 @@ static FLinearColor SampleImage(const FLinearColor* Pixels, int Width, int Heigh
 		Color11 * (FracX1 * FracY1);
 }
 
-static void ResizeImage(const FImage& SrcImage, FImage& DestImage)
+static void ResizeImage(const FImageView & SrcImage, const FImageView & DestImage)
 {
 	// Src and Dest should both now be RGBA32F and Linear gamma
-	const FLinearColor* SrcPixels = SrcImage.AsRGBA32F().GetData();
-	FLinearColor* DestPixels = DestImage.AsRGBA32F().GetData();
+	check( SrcImage.Format == ERawImageFormat::RGBA32F );
+	check( DestImage.Format == ERawImageFormat::RGBA32F );
+	const FLinearColor* SrcPixels = (const FLinearColor*) SrcImage.RawData;
+	FLinearColor* DestPixels = (FLinearColor*) DestImage.RawData;
 	const float DestToSrcScaleX = (float)SrcImage.SizeX / (float)DestImage.SizeX;
 	const float DestToSrcScaleY = (float)SrcImage.SizeY / (float)DestImage.SizeY;
 
+	// @todo Oodle : not a correct bilinear Resize?  missing 0.5 pixel center shift
 	for (int64 DestY = 0; DestY < DestImage.SizeY; ++DestY)
 	{
 		const float SrcY = (float)DestY * DestToSrcScaleY;
@@ -852,27 +855,43 @@ void FImage::CopyTo(FImage& DestImage, ERawImageFormat::Type DestFormat, EGammaS
 		DestGammaSpace = EGammaSpace::sRGB;
 	}
 
-	DestImage.SizeX = SizeX;
-	DestImage.SizeY = SizeY;
-	DestImage.NumSlices = NumSlices;
-	DestImage.Format = DestFormat;
-	DestImage.GammaSpace = DestGammaSpace;
-	InitImageStorage(DestImage);
+	DestImage.Init(SizeX,SizeY,NumSlices,DestFormat,DestGammaSpace);
+	FImageCore::CopyImage(*this, DestImage);
+}
+
+void FImageView::CopyTo(FImage& DestImage, ERawImageFormat::Type DestFormat, EGammaSpace DestGammaSpace) const
+{
+	// if gamma correction is done, it's always *TO* sRGB , not to Pow22
+	// so if Pow22 was requested, change to sRGB
+	// so that Float->int->Float roundtrips correctly
+	if ( DestGammaSpace == EGammaSpace::Pow22 && GammaSpace != EGammaSpace::Pow22 )
+	{
+		// fix call sites that hit this
+		UE_LOG(LogImageCore, Warning, TEXT("Pow22 should not be used as a Dest GammaSpace.  Pow22 Source should encode to sRGB Dest."));
+		DestGammaSpace = EGammaSpace::sRGB;
+	}
+
+	DestImage.Init(SizeX,SizeY,NumSlices,DestFormat,DestGammaSpace);
 	FImageCore::CopyImage(*this, DestImage);
 }
 
 void FImage::ResizeTo(FImage& DestImage, int32 DestSizeX, int32 DestSizeY, ERawImageFormat::Type DestFormat, EGammaSpace DestGammaSpace) const
 {
-	check(NumSlices == 1); // only support 1 slice for now
+	FImageCore::ResizeTo(*this,DestImage,DestSizeX,DestSizeY,DestFormat,DestGammaSpace);
+}
+
+void FImageCore::ResizeTo(const FImageView & SourceImage,FImage& DestImage, int32 DestSizeX, int32 DestSizeY, ERawImageFormat::Type DestFormat, EGammaSpace DestGammaSpace)
+{
+	check(SourceImage.NumSlices == 1); // only support 1 slice for now
 
 	FImage TempSrcImage;
-	const FImage* SrcImagePtr = this;
-	if (Format != ERawImageFormat::RGBA32F)
+	FImageView SrcImageView = SourceImage;
+	if (SourceImage.Format != ERawImageFormat::RGBA32F)
 	{
-		CopyTo(TempSrcImage, ERawImageFormat::RGBA32F, EGammaSpace::Linear);
-		SrcImagePtr = &TempSrcImage;
+		SourceImage.CopyTo(TempSrcImage, ERawImageFormat::RGBA32F, EGammaSpace::Linear);
+		SrcImageView = TempSrcImage;
 	}
-	else if ( GammaSpace != EGammaSpace::Linear )
+	else if ( SourceImage.GammaSpace != EGammaSpace::Linear )
 	{
 		UE_LOG(LogImageCore, Warning, TEXT("Resize from source Format RGBA32F was called but source GammaSpace is not Linear"));
 	}
@@ -892,7 +911,7 @@ void FImage::ResizeTo(FImage& DestImage, int32 DestSizeX, int32 DestSizeY, ERawI
 		DestImage.Format = DestFormat;
 		DestImage.GammaSpace = DestGammaSpace;
 		InitImageStorage(DestImage);
-		ResizeImage(*SrcImagePtr, DestImage);
+		ResizeImage(SrcImageView, DestImage);
 	}
 	else
 	{
@@ -904,7 +923,7 @@ void FImage::ResizeTo(FImage& DestImage, int32 DestSizeX, int32 DestSizeY, ERawI
 		TempDestImage.Format = ERawImageFormat::RGBA32F;
 		TempDestImage.GammaSpace = EGammaSpace::Linear;
 		InitImageStorage(TempDestImage);
-		ResizeImage(*SrcImagePtr, TempDestImage);
+		ResizeImage(SrcImageView, TempDestImage);
 
 		// then convert to dest format/gamma :
 		if ( DestGammaSpace == EGammaSpace::Pow22 )
@@ -1282,7 +1301,7 @@ void FImageCore::SanitizeFloat16AndSetAlphaOpaqueForBC6H(const FImageView & InOu
 }
 
 
-bool FImageCore::DetectAlphaChannel(const FImage& InImage)
+bool FImageCore::DetectAlphaChannel(const FImageView & InImage)
 {
 	TRACE_CPUPROFILER_EVENT_SCOPE(Texture.DetectAlphaChannel);
 
@@ -1381,3 +1400,77 @@ bool FImageCore::DetectAlphaChannel(const FImage& InImage)
 
 	return false;
 }
+
+
+void FImageCore::SetAlphaOpaque(const FImageView & InImage)
+{
+	int64 NumPixels = (int64)InImage.SizeX * InImage.SizeY * InImage.NumSlices;
+
+	if (InImage.Format == ERawImageFormat::BGRA8)
+	{
+		TArrayView64<FColor> SrcColorArray = InImage.AsBGRA8();
+		check(SrcColorArray.Num() == NumPixels);
+
+		FColor* ColorPtr = &SrcColorArray[0];
+		FColor* EndPtr = ColorPtr + SrcColorArray.Num();
+
+		for (; ColorPtr < EndPtr; ++ColorPtr)
+		{
+			ColorPtr->A = 255;
+		}
+	}
+	else if (InImage.Format == ERawImageFormat::RGBA32F)
+	{
+		TArrayView64<FLinearColor> SrcColorArray = InImage.AsRGBA32F();
+		check(SrcColorArray.Num() == NumPixels);
+
+		FLinearColor* ColorPtr = &SrcColorArray[0];
+		FLinearColor* EndPtr = ColorPtr + SrcColorArray.Num();
+
+		for (; ColorPtr < EndPtr; ++ColorPtr)
+		{
+			ColorPtr->A = 1.f;
+		}
+	}
+	else if (InImage.Format == ERawImageFormat::RGBA16)
+	{
+		TArrayView64<uint16> SrcChannelArray = InImage.AsRGBA16();
+		check(SrcChannelArray.Num() == NumPixels * 4);
+
+		uint16* ChannelPtr = &SrcChannelArray[0];
+		uint16* EndPtr = ChannelPtr + SrcChannelArray.Num();
+
+		for (; ChannelPtr < EndPtr; ChannelPtr += 4)
+		{
+			ChannelPtr[3] = 0xFFFF;
+		}
+	}
+	else if (InImage.Format == ERawImageFormat::RGBA16F)
+	{
+		TArrayView64<FFloat16Color> SrcColorArray = InImage.AsRGBA16F();
+		check(SrcColorArray.Num() == NumPixels);
+
+		FFloat16Color* ColorPtr = &SrcColorArray[0];
+		FFloat16Color* EndPtr = ColorPtr + SrcColorArray.Num();
+
+		FFloat16 One(1.f);
+
+		for (; ColorPtr < EndPtr; ++ColorPtr)
+		{
+			ColorPtr->A = One;
+		}
+	}
+	else if (InImage.Format == ERawImageFormat::G8 ||
+		InImage.Format == ERawImageFormat::BGRE8 ||
+		InImage.Format == ERawImageFormat::G16 ||
+		InImage.Format == ERawImageFormat::R16F ||
+		InImage.Format == ERawImageFormat::R32F)
+	{
+		// source image formats don't have alpha
+	}
+	else
+	{
+		// new format ?
+		check(0);
+	}
+}
diff --git a/Engine/Source/Runtime/ImageCore/Public/ImageCore.h b/Engine/Source/Runtime/ImageCore/Public/ImageCore.h
--- a/Engine/Source/Runtime/ImageCore/Public/ImageCore.h
+++ b/Engine/Source/Runtime/ImageCore/Public/ImageCore.h
@@ -35,6 +35,8 @@
 /* Types
  *****************************************************************************/
 
+ struct FImage;
+
 // EGammaSpace is in Color.h
 
 namespace ERawImageFormat
@@ -253,6 +255,80 @@ struct FImageView : public FImageInfo
 
 	// get an FImageView to one 2d slice
 	IMAGECORE_API FImageView GetSlice(int32 SliceIndex) const;
+	
+	/**
+	 * Copies the image to a destination image with the specified format.
+	 *
+	 * @param DestImage - The destination image.
+	 * @param DestFormat - The destination image format.
+	 * @param DestSRGB - Whether the destination image is in SRGB format.
+	 */
+	IMAGECORE_API void CopyTo(FImage& DestImage, ERawImageFormat::Type DestFormat, EGammaSpace DestGammaSpace) const;
+	
+	// CopyTo same format
+	IMAGECORE_API void CopyTo(FImage& DestImage) const
+	{
+		CopyTo(DestImage,Format,GammaSpace);
+	}
+	
+public:
+
+	// Convenience accessors to raw data
+	// these are const member functions because the FImageView object itself is const, but the image it points to may not be
+	
+	TArrayView64<uint8> AsG8() const
+	{
+		check(Format == ERawImageFormat::G8);
+		return { (uint8 *)RawData, GetNumPixels() };
+	}
+
+	TArrayView64<uint16> AsG16() const
+	{
+		check(Format == ERawImageFormat::G16);
+		return { (uint16*)RawData, GetNumPixels() };
+	}
+
+	TArrayView64<FColor> AsBGRA8() const
+	{
+		check(Format == ERawImageFormat::BGRA8);
+		return { (struct FColor*)RawData, GetNumPixels() };
+	}
+
+	TArrayView64<FColor> AsBGRE8() const
+	{
+		check(Format == ERawImageFormat::BGRE8);
+		return { (struct FColor*)RawData, GetNumPixels() };
+	}
+
+	TArrayView64<uint16> AsRGBA16() const
+	{
+		check(Format == ERawImageFormat::RGBA16);
+		return { (uint16*)RawData, GetNumPixels() };
+	}
+
+	TArrayView64<FFloat16Color> AsRGBA16F() const
+	{
+		check(Format == ERawImageFormat::RGBA16F);
+		return { (class FFloat16Color*)RawData, GetNumPixels() };
+	}
+
+	TArrayView64<FLinearColor> AsRGBA32F() const
+	{
+		check(Format == ERawImageFormat::RGBA32F);
+		return { (struct FLinearColor*)RawData, GetNumPixels() };
+	}
+
+	TArrayView64<FFloat16> AsR16F() const
+	{
+		check(Format == ERawImageFormat::R16F);
+		return { (class FFloat16*)RawData, GetNumPixels() };
+	}
+	
+	TArrayView64<float> AsR32F() const
+	{
+		check(Format == ERawImageFormat::R32F);
+		return { (float *)RawData, GetNumPixels() };
+	}
 };
 
 /**
@@ -598,7 +674,26 @@ IMAGECORE_API void SanitizeFloat16AndSetAlphaOpaqueForBC6H(const FImageView & In
  * @param InImage - The image to look for alpha in.
  * @return True if the image both supports and uses the alpha channel. False otherwise.
  */
-IMAGECORE_API bool DetectAlphaChannel(const FImage & InImage);
+IMAGECORE_API bool DetectAlphaChannel(const FImageView & InImage);
+
+/**
+ * Change alpha channel to opaque, if present
+ * 
+ * @param InImage - The image to look for alpha in.
+ */
+IMAGECORE_API void SetAlphaOpaque(const FImageView & InImage);
+
+/**
+	* Copies and resizes the image to a destination image with the specified size and format.
+	* Resize is done using bilinear filtering
+	*
+	* @param DestImage - The destination image.
+	* @param DestSizeX - Width of the resized image
+	* @param DestSizeY - Height of the resized image
+	* @param DestFormat - The destination image format.
+	* @param DestSRGB - Whether the destination image is in SRGB format.
+	*/
+IMAGECORE_API void ResizeTo(const FImageView & SourceImage,FImage& DestImage, int32 DestSizeX, int32 DestSizeY, ERawImageFormat::Type DestFormat, EGammaSpace DestGammaSpace);
 
 };
 
