DDSFile prevent possible integer overflow in ArrayCount * MipCount
also another in Width*Height*Depth*PixelSize

diff --git a/Engine/Source/Runtime/ImageCore/Private/DDSFile.cpp b/Engine/Source/Runtime/ImageCore/Private/DDSFile.cpp
--- a/Engine/Source/Runtime/ImageCore/Private/DDSFile.cpp
+++ b/Engine/Source/Runtime/ImageCore/Private/DDSFile.cpp
@@ -350,6 +350,8 @@ static EDXGIFormat DXGIFormatFromDDS9Header(const FDDSHeaderWithMagic* InDDSHead
 
 static uint32 MipDimension(uint32 dim, uint32 level) 
 {
+	check( level <= FDDSFile::MAX_MIPS_SUPPORTED ); // guaranteed less than 32
+
 	// mip dimensions truncate at every level and bottom out at 1
 	uint32 x = dim >> level;
 	return x ? x : 1;
@@ -419,6 +421,7 @@ EDDSError FDDSFile::Validate() const
 	}
 
 	// Images must not be larger than we support
+	check( MAX_MIPS_SUPPORTED < 32 );
 	const uint32 MaxDimension = (1 << MAX_MIPS_SUPPORTED) - 1; // 1<<k is when we tip over into needing k+1 mip levels, (1<<k)-1 needs just k.
 	if (Width > MaxDimension || Height > MaxDimension || Depth > MaxDimension)
 	{
@@ -432,6 +435,24 @@ EDDSError FDDSFile::Validate() const
 		UE_LOG(LogDDSFile, Warning, TEXT("Invalid mipmap count of %u."), MipCount);
 		return EDDSError::BadMipmapCount;
 	}
+	
+	// ArraySize * MipCount must fit in 32 bits
+	// checking against max dim will ensure that :
+	if ( ArraySize > MaxDimension )
+	{
+		UE_LOG(LogDDSFile, Warning, TEXT("ArraySize %ux of DDS exceed maximum of %u."), ArraySize, MaxDimension);
+		return EDDSError::BadImageDimension;		
+	}
+
+	// all dimensions are limited to 20 bits, so you could get to 1<<60 pixel count
+	//	with 16-byte pixels the bytes per surface could then go to negative in int64
+	uint64 MaxPixelCount = 1ULL<<40; // could be higher and not overflow int64 , but would run out of memory anyway
+	// note the limit at MaxPixelCount is not strict due to use of floats, but doesn't need to be
+	if ( float(Width) * float(Height) * float(Depth) * float(ArraySize) > float(MaxPixelCount) )
+	{
+		UE_LOG(LogDDSFile, Warning, TEXT("Image dimensions %ux%ux%ux%u of DDS exceed maximum pixel count."), Width, Height, Depth, ArraySize);
+		return EDDSError::BadImageDimension;
+	}
 
 	// Mipmaps halve each dimension (rounding down) every step; dimensions that end up at 0
 	// turn into 1. For the mip count to be valid, in the final mip level, at least one
