Perform raw data copy between FImageViews with matching linear pixel formats, even if specified gamma spaces don't match.

diff --git a/Engine/Source/Runtime/ImageCore/Private/ImageCore.cpp b/Engine/Source/Runtime/ImageCore/Private/ImageCore.cpp
--- a/Engine/Source/Runtime/ImageCore/Private/ImageCore.cpp
+++ b/Engine/Source/Runtime/ImageCore/Private/ImageCore.cpp
@@ -204,8 +204,9 @@ IMAGECORE_API void FImageCore::CopyImage(const FImageView & SrcImage,const FImag
 	check(SrcImage.GetNumPixels() == DestImage.GetNumPixels());
 	
 	// short cut fast path identical pixels before we do anything else
+	// note that gamma correction is only performed for U8 formats
 	if (SrcImage.Format == DestImage.Format &&
-		SrcImage.GammaSpace == DestImage.GammaSpace)
+		(SrcImage.GammaSpace == DestImage.GammaSpace || !ERawImageFormat::GetFormatNeedsGammaSpace(SrcImage.Format)))
 	{
 		int64 Bytes = SrcImage.GetImageSizeBytes();
 		check( DestImage.GetImageSizeBytes() == Bytes );
