ImageCore Gamma validation

diff --git a/Engine/Source/Runtime/ImageCore/Private/ImageCore.cpp b/Engine/Source/Runtime/ImageCore/Private/ImageCore.cpp
--- a/Engine/Source/Runtime/ImageCore/Private/ImageCore.cpp
+++ b/Engine/Source/Runtime/ImageCore/Private/ImageCore.cpp
@@ -21,6 +21,8 @@ IMPLEMENT_MODULE(FDefaultModuleImpl, ImageCore);
  */
 static void InitImageStorage(FImage& Image)
 {
+	check( Image.IsImageInfoValid() );
+
 	int64 NumBytes = Image.GetImageSizeBytes();
 	Image.RawData.Empty(NumBytes);
 	Image.RawData.AddUninitialized(NumBytes);
@@ -199,22 +201,50 @@ static const TCHAR * GammaSpaceGetName(EGammaSpace GammaSpace)
  */
 IMAGECORE_API void FImageCore::CopyImage(const FImageView & SrcImage,const FImageView & DestImage)
 {
+	// self-calls before the TRACE_CPUPROFILER_EVENT_SCOPE
+	// if the calling code is correct, these should not be used
+	//	they are a temporary patch in case bad calling code lingers
+	if ( SrcImage.IsGammaCorrected() && ! GetFormatNeedsGammaSpace(SrcImage.Format) )
+	{
+		UE_LOG(LogImageCore,Warning,TEXT("CopyImage: SrcImage has invalid gamma settings %s %s"),
+			ERawImageFormat::GetName(SrcImage.Format),
+			GammaSpaceGetName(SrcImage.GammaSpace));
+
+		// if we're given {F32,sRGB} assume they meant {F32,Linear} and go ahead with the copy
+		FImageView SrcLinear = SrcImage;
+		SrcLinear.GammaSpace = EGammaSpace::Linear;
+		FImageCore::CopyImage(SrcLinear,DestImage);
+		return;
+	}
+	else if ( DestImage.IsGammaCorrected() && ! GetFormatNeedsGammaSpace(DestImage.Format) )
+	{
+		UE_LOG(LogImageCore,Warning,TEXT("CopyImage: DestImage has invalid gamma settings %s %s"),
+			ERawImageFormat::GetName(DestImage.Format),
+			GammaSpaceGetName(DestImage.GammaSpace));
+			
+		// if we're given {F32,sRGB} assume they meant {F32,Linear} and go ahead with the copy
+		FImageView DestLinear = DestImage;
+		DestLinear.GammaSpace = EGammaSpace::Linear;
+		FImageCore::CopyImage(SrcImage,DestLinear);
+		return;
+	}
+
 	TRACE_CPUPROFILER_EVENT_SCOPE(Texture.CopyImage);
 
+	check(SrcImage.IsImageInfoValid());
+	check(DestImage.IsImageInfoValid());
 	check(SrcImage.GetNumPixels() == DestImage.GetNumPixels());
 	
 	// short cut fast path identical pixels before we do anything else
-	// note that gamma correction is only performed for U8 formats
-	if (SrcImage.Format == DestImage.Format &&
-		(SrcImage.GammaSpace == DestImage.GammaSpace || !ERawImageFormat::GetFormatNeedsGammaSpace(SrcImage.Format)))
+	if ( SrcImage.Format == DestImage.Format &&
+		SrcImage.GammaSpace == DestImage.GammaSpace )
 	{
 		int64 Bytes = SrcImage.GetImageSizeBytes();
 		check( DestImage.GetImageSizeBytes() == Bytes );
 		memcpy(DestImage.RawData,SrcImage.RawData,Bytes);
-
 		return;
 	}
-
+	
 	UE_LOG(LogImageCore,Verbose,TEXT("CopyImage: %s %s -> %s %s %dx%d"),
 		ERawImageFormat::GetName(SrcImage.Format),
 		GammaSpaceGetName(SrcImage.GammaSpace),
@@ -223,6 +253,7 @@ IMAGECORE_API void FImageCore::CopyImage(const FImageView & SrcImage,const FImag
 		SrcImage.SizeX,SrcImage.SizeY);
 
 	// bDestIsGammaCorrected for Dest (Pow22 and sRGB) will encode to sRGB
+	// note that gamma correction is only performed for U8 formats
 	const bool bDestIsGammaCorrected = DestImage.IsGammaCorrected();
 	const int64 NumTexels = SrcImage.GetNumPixels();
 	int64 TexelsPerJob;
diff --git a/Engine/Source/Runtime/ImageCore/Public/ImageCore.h b/Engine/Source/Runtime/ImageCore/Public/ImageCore.h
--- a/Engine/Source/Runtime/ImageCore/Public/ImageCore.h
+++ b/Engine/Source/Runtime/ImageCore/Public/ImageCore.h
@@ -142,6 +142,15 @@ struct FImageInfo
 	{
 	}
 		
+	FORCEINLINE bool IsImageInfoValid() const
+	{
+		if ( SizeX < 0 || SizeY < 0 || NumSlices < 0 ) return false;
+		if ( Format == ERawImageFormat::Invalid ) return false;
+		if ( GammaSpace == EGammaSpace::Invalid ) return false;
+		if ( ! GetFormatNeedsGammaSpace(Format) && GammaSpace != EGammaSpace::Linear ) return false;
+		return true;
+	}
+
 	FORCEINLINE bool IsGammaCorrected() const
 	{
 		return GammaSpace != EGammaSpace::Linear;
