ImageCore.cpp : comment out check until cause is found

diff --git a/Engine/Source/Runtime/ImageCore/Private/ImageCore.cpp b/Engine/Source/Runtime/ImageCore/Private/ImageCore.cpp
--- a/Engine/Source/Runtime/ImageCore/Private/ImageCore.cpp
+++ b/Engine/Source/Runtime/ImageCore/Private/ImageCore.cpp
@@ -21,7 +21,7 @@ IMPLEMENT_MODULE(FDefaultModuleImpl, ImageCore);
  */
 static void InitImageStorage(FImage& Image)
 {
-	check( Image.IsImageInfoValid() );
+	//check( Image.IsImageInfoValid() );
 
 	int64 NumBytes = Image.GetImageSizeBytes();
 	Image.RawData.Empty(NumBytes);
