TCheckedSignedInt -> TGuardedSignedInt

Derived types also renamed from "checked" to "guarded" to avoid confusion with check()
Remove "Invalid()" static method. The coding style compliant "GetInvalid()" is too confusing (it sounds like it might be !IsValid()), just use the default ctor.
GetValue() -> GetChecked() with user-specified default value. This is like Get() but with a check(IsValid()) in front (to make sure results are always defined, even in non-checked builds).
Add natvis for TGuardedSignedInt
Various small cleanups: constants const->constexpr, simplify from-int ctor
Only allow arithmetic ops with maching SignedType
Move to ImageCore for now

diff --git a/Engine/Plugins/Interchange/Runtime/Source/Import/Private/Texture/InterchangeSlicedTexturePayloadData.cpp b/Engine/Plugins/Interchange/Runtime/Source/Import/Private/Texture/InterchangeSlicedTexturePayloadData.cpp
--- a/Engine/Plugins/Interchange/Runtime/Source/Import/Private/Texture/InterchangeSlicedTexturePayloadData.cpp
+++ b/Engine/Plugins/Interchange/Runtime/Source/Import/Private/Texture/InterchangeSlicedTexturePayloadData.cpp
@@ -2,7 +2,7 @@
 
 #include "Texture/InterchangeSlicedTexturePayloadData.h"
 
-#include "Math/CheckedInt.h"
+#include "GuardedInt.h"
 #include "CoreMinimal.h"
 
 void UE::Interchange::FImportSlicedImage::Init(int32 InSizeX, int32 InSizeY, int32 InNumSlice, int32 InNumMips, ETextureSourceFormat InFormat, bool InSRGB)
@@ -54,7 +54,7 @@ int64 UE::Interchange::FImportSlicedImage::GetMipSize(int32 InMipIndex) const
 	int32 MipSizeY = FMath::Max<int32>(SizeY >> InMipIndex, 1);
 	int32 MipSizeZ = bIsVolume ? FMath::Max(NumSlice >> InMipIndex, 1) : 1;
 
-	FCheckedInt64 MipSize = FCheckedInt64(MipSizeX) * MipSizeY * MipSizeZ * BytesPerPixel;
+	FGuardedInt64 MipSize = FGuardedInt64(MipSizeX) * MipSizeY * MipSizeZ * BytesPerPixel;
 	checkf(MipSize.IsValid(), TEXT("Invalid (overflowing) mip sizes made it in to FImportSlicedImage::GetMipSize! Check import locations for mip size validation"));
 
 	return MipSize.Get(0);
diff --git a/Engine/Source/Runtime/Core/Private/Tests/Math/UnrealMathTest.cpp b/Engine/Source/Runtime/Core/Private/Tests/Math/UnrealMathTest.cpp
--- a/Engine/Source/Runtime/Core/Private/Tests/Math/UnrealMathTest.cpp
+++ b/Engine/Source/Runtime/Core/Private/Tests/Math/UnrealMathTest.cpp
@@ -17,7 +17,6 @@
 #include "Math/Quat.h"
 #include "Math/QuatRotationTranslationMatrix.h"
 #include "Math/Color.h"
-#include "Math/CheckedInt.h"
 #include "Misc/AutomationTest.h"
 #include "Async/ParallelFor.h"
 #include "Misc/ScopeLock.h"
@@ -4809,137 +4808,4 @@ bool FColorConversionHeavyTest::RunTest(const FString& Parameters)
 	return true;
 }
 
-IMPLEMENT_SIMPLE_AUTOMATION_TEST(FMathCheckedIntTest, "System.Core.Math.CheckedInt", EAutomationTestFlags::ApplicationContextMask | EAutomationTestFlags::EngineFilter)
-bool FMathCheckedIntTest::RunTest(const FString& Parameters)
-{
-	// Do an exhaustive test of all checked signed int operations with int8s. This is a small enough space of possibilities
-	// that it's easy to try all options exhaustively, and it's likewise easy to compute reference results in int32 arithmetic
-	// without having to worry about overflows. TCheckedSignedInt uses the same logic and implementation for all type
-	// instantiations (as of this writing), so this is a good way to get coverage on all relevant code paths.
-	using FCheckedInt8 = TCheckedSignedInt<int8>;
-	int NumErrorsFound = 0;
-	const int NumErrorsMax = 100; // Stop reporting errors after this in case something is completely busted.
-
-	// Helper functions to log test results
-	auto CheckArithResult = [this, &NumErrorsFound, NumErrorsMax](const TCHAR* Operation, int32 A, int32 B, FCheckedInt8 CheckedResult, int32 UncheckedValue)
-	{
-		FCheckedInt8 UncheckedResult{ UncheckedValue };
-
-		// If either operand was invalid, the unchecked result is also invalid
-		if (!FCheckedInt8(A).IsValid() || !FCheckedInt8(B).IsValid())
-		{
-			UncheckedResult = FCheckedInt8::Invalid();
-		}
-
-		if (CheckedResult != UncheckedResult && NumErrorsFound < NumErrorsMax)
-		{
-			NumErrorsFound++;
-
-			const int UncheckedNum = UncheckedResult.Get(-128);
-			const int CheckedNum = CheckedResult.Get(-128);
-
-			AddError(FString::Printf(TEXT("Arith %s A=%d B=%d Expected=(%d,Valid=%d) Got=(%d,Valid=%d)"),
-				Operation, A, B, UncheckedNum, UncheckedResult.IsValid(), CheckedNum, CheckedResult.IsValid()));
-		}
-	};
-
-	auto CheckCompareResult = [this, &NumErrorsFound, NumErrorsMax](const TCHAR* Operation, int32 A, int32 B, bool bAcceptInvalid, bool bTestResult, bool bRefResult)
-	{
-		// If either operands was invalid, override the ref result
-		if (!FCheckedInt8(A).IsValid() || !FCheckedInt8(B).IsValid())
-			bRefResult = bAcceptInvalid;
-
-		if (bTestResult != bRefResult && NumErrorsFound < NumErrorsMax)
-		{
-			NumErrorsFound++;
-
-			AddError(FString::Printf(TEXT("Cmp %s A=%d B=%d ref=%d test=%d\n"), Operation, A, B, bRefResult, bTestResult));
-		}
-	};
-
-	// Note: we include 128 in the ranges here to explicitly try out-of-range (invalid) values.
-	for (int32 A = -128; A <= 128; ++A)
-	{
-		for (int32 B = -128; B <= 128; ++B)
-		{
-			FCheckedInt8 CheckedA{ A };
-			FCheckedInt8 CheckedB{ B };
-
-			// Compute reference results for the operations that are not defined everywhere
-			int32 ADivB = -256;
-			int32 AModB = -256;
-
-			// -128 / -1 is undefined in int8 arithmetic because it overflows; when we're
-			// computing it with int32s, it of course works just fine, but we want to make sure
-			// our test results are correct.
-			if (B != 0 && (A != -128 || B != -1))
-			{
-				ADivB = A / B;
-				AModB = A % B;
-			}
-
-			// Shifts are only defined for in-range shift amounts, which for our checked ints
-			// is [0,NumBits-1]. This is tighter than usual C++ arithmetic for the FCheckedInt8 we
-			// use for testing since normal C++ promotes to int and shifts down (so usually [0,31]
-			// count limit) so we need to test for these explicitly.
-			int32 ALshB = -256;
-			int32 ARshB = -256;
-			if (B >= 0 && B <= 7)
-			{
-				ALshB = A << B;
-				ARshB = A >> B;
-			}
-
-			CheckArithResult(TEXT("add1"), A, B, CheckedA + CheckedB, A + B);
-			CheckArithResult(TEXT("add2"), A, B, CheckedA + B, A + B);
-			CheckArithResult(TEXT("add3"), A, B, A + CheckedB, A + B);
-			CheckArithResult(TEXT("sub1"), A, B, CheckedA - CheckedB, A - B);
-			CheckArithResult(TEXT("sub2"), A, B, CheckedA - B, A - B);
-			CheckArithResult(TEXT("sub3"), A, B, A - CheckedB, A - B);
-			CheckArithResult(TEXT("mul1"), A, B, CheckedA * CheckedB, A * B);
-			CheckArithResult(TEXT("mul2"), A, B, CheckedA * B, A * B);
-			CheckArithResult(TEXT("mul3"), A, B, A * CheckedB, A * B);
-			CheckArithResult(TEXT("div1"), A, B, CheckedA / CheckedB, ADivB);
-			CheckArithResult(TEXT("div2"), A, B, CheckedA / B, ADivB);
-			CheckArithResult(TEXT("div3"), A, B, A / CheckedB, ADivB);
-			CheckArithResult(TEXT("mod1"), A, B, CheckedA % CheckedB, AModB);
-			CheckArithResult(TEXT("mod2"), A, B, CheckedA % B, AModB);
-			CheckArithResult(TEXT("mod3"), A, B, A % CheckedB, AModB);
-			CheckArithResult(TEXT("lsh1"), A, B, CheckedA << CheckedB, ALshB);
-			CheckArithResult(TEXT("lsh2"), A, B, CheckedA << B, ALshB);
-			CheckArithResult(TEXT("lsh3"), A, B, A << CheckedB, ALshB);
-			CheckArithResult(TEXT("rsh1"), A, B, CheckedA >> CheckedB, ARshB);
-			CheckArithResult(TEXT("rsh2"), A, B, CheckedA >> B, ARshB);
-			CheckArithResult(TEXT("rsh3"), A, B, A >> CheckedB, ARshB);
-			CheckArithResult(TEXT("neg"), A, A, -CheckedA, -A);
-			CheckArithResult(TEXT("abs"), A, A, CheckedA.Abs(), FMath::Abs(A));
-			
-			CheckCompareResult(TEXT("valid_lt1"), A, B, false, CheckedA.ValidAndLessThan(CheckedB), A < B);
-			CheckCompareResult(TEXT("valid_lt2"), A, B, false, CheckedA.ValidAndLessThan(B), A < B);
-			CheckCompareResult(TEXT("valid_le1"), A, B, false, CheckedA.ValidAndLessOrEqual(CheckedB), A <= B);
-			CheckCompareResult(TEXT("valid_le1"), A, B, false, CheckedA.ValidAndLessOrEqual(B), A <= B);
-			CheckCompareResult(TEXT("valid_gt1"), A, B, false, CheckedA.ValidAndGreaterThan(CheckedB), A > B);
-			CheckCompareResult(TEXT("valid_gt1"), A, B, false, CheckedA.ValidAndGreaterThan(B), A > B);
-			CheckCompareResult(TEXT("valid_ge1"), A, B, false, CheckedA.ValidAndGreaterOrEqual(CheckedB), A >= B);
-			CheckCompareResult(TEXT("valid_ge2"), A, B, false, CheckedA.ValidAndGreaterOrEqual(B), A >= B);
-			
-			CheckCompareResult(TEXT("invalid_lt1"), A, B, true, CheckedA.InvalidOrLessThan(CheckedB), A < B);
-			CheckCompareResult(TEXT("invalid_lt2"), A, B, true, CheckedA.InvalidOrLessThan(B), A < B);
-			CheckCompareResult(TEXT("invalid_le1"), A, B, true, CheckedA.InvalidOrLessOrEqual(CheckedB), A <= B);
-			CheckCompareResult(TEXT("invalid_le2"), A, B, true, CheckedA.InvalidOrLessOrEqual(B), A <= B);
-			CheckCompareResult(TEXT("invalid_gt1"), A, B, true, CheckedA.InvalidOrGreaterThan(CheckedB), A > B);
-			CheckCompareResult(TEXT("invalid_gt2"), A, B, true, CheckedA.InvalidOrGreaterThan(B), A > B);
-			CheckCompareResult(TEXT("invalid_ge1"), A, B, true, CheckedA.InvalidOrGreaterOrEqual(CheckedB), A >= B);
-			CheckCompareResult(TEXT("invalid_ge2"), A, B, true, CheckedA.InvalidOrGreaterOrEqual(B), A >= B);
-		}
-	}
-
-	if (NumErrorsFound > NumErrorsMax)
-	{
-		AddError(FString::Printf(TEXT("Too many errors (%d), stopped reporting after %d instances.\n"), NumErrorsFound, NumErrorsMax));
-	}
-
-	return NumErrorsFound == 0;
-}
-
 #endif //WITH_DEV_AUTOMATION_TESTS
diff --git a/Engine/Source/Runtime/Core/Public/Math/CheckedInt.h b/Engine/Source/Runtime/Core/Public/Math/CheckedInt.h
deleted file mode 100644
diff --git a/Engine/Source/Runtime/Engine/Private/Texture.cpp b/Engine/Source/Runtime/Engine/Private/Texture.cpp
--- a/Engine/Source/Runtime/Engine/Private/Texture.cpp
+++ b/Engine/Source/Runtime/Engine/Private/Texture.cpp
@@ -2,7 +2,7 @@
 
 #include "Engine/Texture.h"
 
-#include "Math/CheckedInt.h"
+#include "GuardedInt.h"
 #include "Misc/App.h"
 #include "Modules/ModuleManager.h"
 #include "Materials/MaterialInterface.h"
@@ -2728,14 +2728,14 @@ int64 FTextureSource::CalcLayerSize(const FTextureSourceBlock& Block, int32 Laye
 	int64 BytesPerPixel = GetBytesPerPixel(LayerIndex);
 
 	// This is used for memory allocation, so use FCheckedInt to rigorously check against overflow issues.
-	FCheckedInt64 TotalSize(0);
+	FGuardedInt64 TotalSize(0);
 	for (int32 MipIndex = 0; MipIndex < Block.NumMips; ++MipIndex)
 	{
 		int32 MipSizeX = FMath::Max<int32>(Block.SizeX >> MipIndex, 1);
 		int32 MipSizeY = FMath::Max<int32>(Block.SizeY >> MipIndex, 1);
 		int32 MipSizeZ = GetMippedNumSlices(Block.NumSlices,MipIndex);
 
-		TotalSize += FCheckedInt64(MipSizeX) * MipSizeY * MipSizeZ * BytesPerPixel;
+		TotalSize += FGuardedInt64(MipSizeX) * MipSizeY * MipSizeZ * BytesPerPixel;
 	}
 
 	checkf(TotalSize.IsValid(), TEXT("Invalid (overflowing) mip sizes made it in to FTextureSource::CalcLayerSize! Check import locations for mip size validation"));
@@ -2749,7 +2749,7 @@ int64 FTextureSource::CalcMipOffset(int32 BlockIndex, int32 LayerIndex, int32 Of
 	check(OffsetToMipIndex < Block.NumMips);
 
 	// This is used for memory indexing, so use FCheckedInt to rigorously check against overflow issues.
-	FCheckedInt64 MipOffset(BlockDataOffsets[BlockIndex]);
+	FGuardedInt64 MipOffset(BlockDataOffsets[BlockIndex]);
 
 	// Skip over the initial layers within the tile
 	for (int i = 0; i < LayerIndex; ++i)
@@ -2765,7 +2765,7 @@ int64 FTextureSource::CalcMipOffset(int32 BlockIndex, int32 LayerIndex, int32 Of
 		int32 MipSizeY = FMath::Max<int32>(Block.SizeY >> MipIndex, 1);
 		int32 MipSizeZ = GetMippedNumSlices(Block.NumSlices,MipIndex);
 
-		MipOffset += FCheckedInt64(MipSizeX) * MipSizeY * MipSizeZ * BytesPerPixel;
+		MipOffset += FGuardedInt64(MipSizeX) * MipSizeY * MipSizeZ * BytesPerPixel;
 	}
 
 	checkf(MipOffset.IsValid(), TEXT("Invalid (overflowing) mip sizes made it in to FTextureSource::CalcMipOffset! Check import locations for mip size validation"));
diff --git a/Engine/Source/Runtime/ImageCore/Private/ImageCoreTest.cpp b/Engine/Source/Runtime/ImageCore/Private/ImageCoreTest.cpp
new file mode 100644
--- /dev/null
+++ b/Engine/Source/Runtime/ImageCore/Private/ImageCoreTest.cpp
@@ -0,0 +1,155 @@
+// Copyright Epic Games, Inc. All Rights Reserved.
+
+#include "GuardedInt.h"
+#include "Misc/AutomationTest.h"
+
+#if WITH_DEV_AUTOMATION_TESTS
+
+IMPLEMENT_SIMPLE_AUTOMATION_TEST(FGuardedIntTest, "System.ImageCore.GuardedInt", EAutomationTestFlags::ApplicationContextMask | EAutomationTestFlags::EngineFilter)
+bool FGuardedIntTest::RunTest(const FString& Parameters)
+{
+	// Do an exhaustive test of all checked signed int operations with int8s. This is a small enough space of possibilities
+	// that it's easy to try all options exhaustively, and it's likewise easy to compute reference results in int32 arithmetic
+	// without having to worry about overflows. TCheckedSignedInt uses the same logic and implementation for all type
+	// instantiations (as of this writing), so this is a good way to get coverage on all relevant code paths.
+	using FGuardedInt8 = TGuardedSignedInt<int8>;
+	int NumErrorsFound = 0;
+	const int NumErrorsMax = 100; // Stop reporting errors after this in case something is completely busted.
+
+	// Helper functions to log test results
+	auto CheckArithResult = [this, &NumErrorsFound, NumErrorsMax](const TCHAR* Operation, int32 A, int32 B, FGuardedInt8 CheckedResult, int32 UncheckedValue)
+	{
+		FGuardedInt8 UncheckedResult{ UncheckedValue };
+
+		// If either operand was invalid, the unchecked result is also invalid
+		if (!FGuardedInt8(A).IsValid() || !FGuardedInt8(B).IsValid())
+		{
+			UncheckedResult = FGuardedInt8();
+		}
+
+		if (CheckedResult != UncheckedResult && NumErrorsFound < NumErrorsMax)
+		{
+			NumErrorsFound++;
+
+			const int UncheckedNum = UncheckedResult.Get(-128);
+			const int CheckedNum = CheckedResult.Get(-128);
+
+			AddError(FString::Printf(TEXT("Arith %s A=%d B=%d Expected=(%d,Valid=%d) Got=(%d,Valid=%d)"),
+				Operation, A, B, UncheckedNum, UncheckedResult.IsValid(), CheckedNum, CheckedResult.IsValid()));
+		}
+	};
+
+	auto CheckCompareResult = [this, &NumErrorsFound, NumErrorsMax](const TCHAR* Operation, int32 A, int32 B, bool bAcceptInvalid, bool bTestResult, bool bRefResult)
+	{
+		// If either operands was invalid, override the ref result
+		if (!FGuardedInt8(A).IsValid() || !FGuardedInt8(B).IsValid())
+			bRefResult = bAcceptInvalid;
+
+		if (bTestResult != bRefResult && NumErrorsFound < NumErrorsMax)
+		{
+			NumErrorsFound++;
+
+			AddError(FString::Printf(TEXT("Cmp %s A=%d B=%d ref=%d test=%d\n"), Operation, A, B, bRefResult, bTestResult));
+		}
+	};
+
+	// Note: we include 128 in the ranges here to explicitly try out-of-range (invalid) values.
+	for (int32 A = -128; A <= 128; ++A)
+	{
+		for (int32 B = -128; B <= 128; ++B)
+		{
+			FGuardedInt8 CheckedA{ A };
+			FGuardedInt8 CheckedB{ B };
+
+			// Compute reference results for the operations that are not defined everywhere
+			int32 ADivB = -256;
+			int32 AModB = -256;
+
+			// -128 / -1 is undefined in int8 arithmetic because it overflows; when we're
+			// computing it with int32s, it of course works just fine, but we want to make sure
+			// our test results are correct.
+			if (B != 0 && (A != -128 || B != -1))
+			{
+				ADivB = A / B;
+				AModB = A % B;
+			}
+
+			// Shifts are only defined for in-range shift amounts, which for our checked ints
+			// is [0,NumBits-1]. This is tighter than usual C++ arithmetic for the FGuardedInt8 we
+			// use for testing since normal C++ promotes to int and shifts down (so usually [0,31]
+			// count limit) so we need to test for these explicitly.
+			int32 ALshB = -256;
+			int32 ARshB = -256;
+			if (B >= 0 && B <= 7)
+			{
+				ALshB = A << B;
+				ARshB = A >> B;
+			}
+
+			CheckArithResult(TEXT("add1"), A, B, CheckedA + CheckedB, A + B);
+			CheckArithResult(TEXT("sub1"), A, B, CheckedA - CheckedB, A - B);
+			CheckArithResult(TEXT("mul1"), A, B, CheckedA * CheckedB, A * B);
+			CheckArithResult(TEXT("div1"), A, B, CheckedA / CheckedB, ADivB);
+			CheckArithResult(TEXT("mod1"), A, B, CheckedA % CheckedB, AModB);
+			CheckArithResult(TEXT("lsh1"), A, B, CheckedA << CheckedB, ALshB);
+			CheckArithResult(TEXT("rsh1"), A, B, CheckedA >> CheckedB, ARshB);
+			CheckArithResult(TEXT("neg"), A, A, -CheckedA, -A);
+			CheckArithResult(TEXT("abs"), A, A, CheckedA.Abs(), FMath::Abs(A));
+
+			// Mixed-type arithmetic needs int8 arguments, which we can only actually do when the values in question _are_ int8s.
+			// (i.e. we can't do this with our intentional out-of-range value of 128.)
+			if (CheckedB.IsValid())
+			{
+				int8 BAs8 = (int8)B;
+
+				CheckArithResult(TEXT("add2"), A, B, CheckedA + BAs8, A + B);
+				CheckArithResult(TEXT("sub2"), A, B, CheckedA - BAs8, A - B);
+				CheckArithResult(TEXT("mul2"), A, B, CheckedA * BAs8, A * B);
+				CheckArithResult(TEXT("div2"), A, B, CheckedA / BAs8, ADivB);
+				CheckArithResult(TEXT("mod2"), A, B, CheckedA % BAs8, AModB);
+				CheckArithResult(TEXT("lsh2"), A, B, CheckedA << BAs8, ALshB);
+				CheckArithResult(TEXT("rsh2"), A, B, CheckedA >> BAs8, ARshB);
+			}
+
+			if (CheckedA.IsValid())
+			{
+				int8 AAs8 = (int8)A;
+
+				CheckArithResult(TEXT("add3"), A, B, AAs8 + CheckedB, A + B);
+				CheckArithResult(TEXT("sub3"), A, B, AAs8 - CheckedB, A - B);
+				CheckArithResult(TEXT("mul3"), A, B, AAs8 * CheckedB, A * B);
+				CheckArithResult(TEXT("div3"), A, B, AAs8 / CheckedB, ADivB);
+				CheckArithResult(TEXT("mod3"), A, B, AAs8 % CheckedB, AModB);
+				CheckArithResult(TEXT("lsh3"), A, B, AAs8 << CheckedB, ALshB);
+				CheckArithResult(TEXT("rsh3"), A, B, AAs8 >> CheckedB, ARshB);
+			}
+
+			CheckCompareResult(TEXT("valid_lt1"), A, B, false, CheckedA.ValidAndLessThan(CheckedB), A < B);
+			CheckCompareResult(TEXT("valid_lt2"), A, B, false, CheckedA.ValidAndLessThan(B), A < B);
+			CheckCompareResult(TEXT("valid_le1"), A, B, false, CheckedA.ValidAndLessOrEqual(CheckedB), A <= B);
+			CheckCompareResult(TEXT("valid_le1"), A, B, false, CheckedA.ValidAndLessOrEqual(B), A <= B);
+			CheckCompareResult(TEXT("valid_gt1"), A, B, false, CheckedA.ValidAndGreaterThan(CheckedB), A > B);
+			CheckCompareResult(TEXT("valid_gt1"), A, B, false, CheckedA.ValidAndGreaterThan(B), A > B);
+			CheckCompareResult(TEXT("valid_ge1"), A, B, false, CheckedA.ValidAndGreaterOrEqual(CheckedB), A >= B);
+			CheckCompareResult(TEXT("valid_ge2"), A, B, false, CheckedA.ValidAndGreaterOrEqual(B), A >= B);
+
+			CheckCompareResult(TEXT("invalid_lt1"), A, B, true, CheckedA.InvalidOrLessThan(CheckedB), A < B);
+			CheckCompareResult(TEXT("invalid_lt2"), A, B, true, CheckedA.InvalidOrLessThan(B), A < B);
+			CheckCompareResult(TEXT("invalid_le1"), A, B, true, CheckedA.InvalidOrLessOrEqual(CheckedB), A <= B);
+			CheckCompareResult(TEXT("invalid_le2"), A, B, true, CheckedA.InvalidOrLessOrEqual(B), A <= B);
+			CheckCompareResult(TEXT("invalid_gt1"), A, B, true, CheckedA.InvalidOrGreaterThan(CheckedB), A > B);
+			CheckCompareResult(TEXT("invalid_gt2"), A, B, true, CheckedA.InvalidOrGreaterThan(B), A > B);
+			CheckCompareResult(TEXT("invalid_ge1"), A, B, true, CheckedA.InvalidOrGreaterOrEqual(CheckedB), A >= B);
+			CheckCompareResult(TEXT("invalid_ge2"), A, B, true, CheckedA.InvalidOrGreaterOrEqual(B), A >= B);
+		}
+	}
+
+	if (NumErrorsFound > NumErrorsMax)
+	{
+		AddError(FString::Printf(TEXT("Too many errors (%d), stopped reporting after %d instances.\n"), NumErrorsFound, NumErrorsMax));
+	}
+
+	return NumErrorsFound == 0;
+}
+
+#endif
\ No newline at end of file
diff --git a/Engine/Source/Runtime/ImageCore/Public/GuardedInt.h b/Engine/Source/Runtime/ImageCore/Public/GuardedInt.h
new file mode 100644
--- /dev/null
+++ b/Engine/Source/Runtime/ImageCore/Public/GuardedInt.h
@@ -0,0 +1,450 @@
+// Copyright Epic Games, Inc. All Rights Reserved.
+
+#pragma once
+
+#include "CoreTypes.h"
+#include "Math/NumericLimits.h"
+#include <type_traits>
+
+/**
+ * Overflow- and error-checked integer. For integer arithmetic on data from untrusted sources (like imported files),
+ * especially when doing size computations. Also checks for division by zero and invalid shift amounts.
+ *
+ * You're not meant to use this directly. Use FGuardedInt32 or FGuardedInt64 (defined below). A typical use case
+ * would be:
+ * 
+ *   FGuardedInt64 NumBytes = FGuardedInt32(Width) * Height * BytesPerPixel;
+ *   if (NumBytes.InvalidOrGreaterThan(SizeLimit))
+ *   {
+ *       // Report error.
+ *   }
+ *   int64 NumBytesValidated = NumBytes.GetValue();
+ *
+ * This is a template meant to be instantiated on top of regular basic integer types. The code is written
+ * so the logic is integer-size agnostic and uses just regular C++ arithmetic operations. It is assumed
+ * to run on a two's complement integer platform (which is all we support, and as of C++20 is contractual).
+ * You should generally use the specializations FGuardedInt32 and FGuardedInt64 below.
+ *
+ * Checked integers keep both the integer value and a "valid" flag. Default-constructed guarded ints
+ * are invalid, and guarded integers constructed from an integer value are valid and hold that value.
+ * Guarded integers are somewhat analogous to a TOptional<SignedType> in semantics, and borrow some of
+ * the function names.
+ *
+ * The main feature of guarded integers is that all arithmetic on them is overflow-checked. Any arithmetic
+ * involving guarded integers results in a guarded integer, and any arithmetic involving invalid values,
+ * or arithmetic resulting in overflows or other errors (such as division by zero) likewise results in
+ * an invalid value. The idea is that integer arithmetic using guarded integers should be possible to
+ * write very straightforwardly and without having to consider any of these special cases; if any error
+ * occurred along the way, the result will be invalid. These invalid values can then be checked for and
+ * handled right when the result is converted back to a regular integer.
+ *
+ * Some compilers provide built-ins for overflow-checked integer arithmetic for some types. We could
+ * eventually use this (it's especially interesting for multiplications, since our current overflow-checking
+ * algorithm is fairly expensive), but a big benefit of the current approach is that it uses nothing but
+ * regular arithmetic and is type-agnostic. In particular, this makes it possible to check this implementation
+ * exhaustively against a known-good reference for small integer types such as int8. It is much trickier and
+ * more subtle to do good testing for larger integer types where that brute-force approach is not practical.
+ * As-is, the current approach is not the fastest, but it's not presently intended to be used in contexts
+ * where speed of arithmetic operations is a major concern.
+ */
+template<typename SignedType>
+class TGuardedSignedInt
+{
+private:
+	static_assert(std::is_integral_v<SignedType> && std::is_signed_v<SignedType>, "Only defined for signed ints");
+	typedef typename std::make_unsigned_t<SignedType> UnsignedType;
+
+	static constexpr SignedType MinValue = TNumericLimits<SignedType>::Min();
+	static constexpr SignedType MaxValue = TNumericLimits<SignedType>::Max();
+	static constexpr SignedType NumBits = SignedType((sizeof(SignedType) / sizeof(char)) * 8); // Using sizeof to guess the bit count, ugh.
+	static constexpr UnsignedType UnsignedMSB = (UnsignedType)MinValue; // Assuming two's complement
+
+	SignedType Value = 0;
+	bool bIsValid = false;
+
+public:
+	/** Construct a TGuardedSignedInt with an invalid value. */
+	TGuardedSignedInt() = default;
+
+	/** Construct a TGuardedSignedInt from a regular signed integer value. If it's out of range, it results in an invalid value. */
+	template<
+		typename ValueType,
+		std::enable_if_t<std::is_integral_v<ValueType>>* = nullptr
+	>
+	explicit TGuardedSignedInt(ValueType InValue)
+		: Value((SignedType)InValue), bIsValid(false)
+	{
+		if constexpr (std::is_signed_v<ValueType>)
+		{
+			bIsValid = (InValue >= MinValue && InValue <= MaxValue);
+		}
+		else
+		{
+			bIsValid = InValue <= UnsignedType(MaxValue);
+		}
+	}
+
+	/** Copy-construct a TGuardedSignedInt from another of matching type. */
+	TGuardedSignedInt(const TGuardedSignedInt& Other) = default;
+
+	/** Assign a TGuardedSignedInt to another. */
+	TGuardedSignedInt& operator=(const TGuardedSignedInt& Other) = default;
+
+	/** @return true if current value is valid (assigned and no overflows or other errors occurred), false otherwise. */
+	bool IsValid() const { return bIsValid; }
+
+	/** @return The value if valid, DefaultValue otherwise. */
+	const SignedType Get(const SignedType DefaultValue) const
+	{
+		return IsValid() ? Value : DefaultValue;
+	}
+
+	/** @return Returns the value if valid, DefaultValue otherwise, but also check()s that the value is valid.
+	 * Intended for cases where the value is not expected to be invalid.
+	 */
+	const SignedType GetChecked(const SignedType DefaultValue = 0) const
+	{
+		checkf(IsValid(), "Invalid value in TGuardedSignedInt::GetChecked.");
+		return Get(DefaultValue);
+	}
+
+	/** @return true if *this and Other are either both invalid or both valid and have the same value, false otherwise. */
+	bool operator ==(const TGuardedSignedInt Other) const
+	{
+		if (bIsValid != Other.bIsValid)
+		{
+			return false;
+		}
+
+		return !bIsValid || Value == Other.Value;
+	}
+
+	/** @return true if *this and Other either have different "valid" states or are both valid and have different values, false otherwise (logical negation of ==). */
+	bool operator !=(const TGuardedSignedInt Other) const
+	{
+		if (bIsValid != Other.bIsValid)
+		{
+			return true;
+		}
+
+		return bIsValid && Value != Other.Value;
+	}
+
+	// There are intentionally no overloads for the ordered comparison operators, because we have
+	// to decide what to do about validity as well. Instead, do this.
+
+	/** @return true if *this and Other are both valid so they can be compared. */
+	bool ComparisonValid(const TGuardedSignedInt Other) const { return bIsValid && Other.bIsValid; }
+
+	/** @return true if *this and Other are both valid and *this is less than Other. */
+	template<typename ValueType>
+	bool ValidAndLessThan(const ValueType Other) const
+	{
+		TGuardedSignedInt CheckedOther{ Other };
+		return ComparisonValid(CheckedOther) && Value < CheckedOther.Value;
+	}
+
+	/** @return true if *this and Other are both valid and *this is less than or equal to Other. */
+	template<typename ValueType>
+	bool ValidAndLessOrEqual(const ValueType Other) const
+	{
+		TGuardedSignedInt CheckedOther{ Other };
+		return ComparisonValid(CheckedOther) && Value <= CheckedOther.Value;
+	}
+
+	/** @return true if *this and Other are both valid and *this is greater than Other. */
+	template<typename ValueType>
+	bool ValidAndGreaterThan(const ValueType Other) const
+	{
+		TGuardedSignedInt CheckedOther{ Other };
+		return ComparisonValid(CheckedOther) && Value > CheckedOther.Value;
+	}
+
+	/** @return true if *this and Other are both valid and *this is greater than or equal to Other. */
+	template<typename ValueType>
+	bool ValidAndGreaterOrEqual(const ValueType Other) const
+	{
+		TGuardedSignedInt CheckedOther{ Other };
+		return ComparisonValid(CheckedOther) && Value >= CheckedOther.Value;
+	}
+
+	/** @return true if either of *this or Other are invalid or *this is less than Other. */
+	template<typename ValueType>
+	bool InvalidOrLessThan(const ValueType Other) const
+	{
+		TGuardedSignedInt CheckedOther{ Other };
+		return !ComparisonValid(CheckedOther) || Value < CheckedOther.Value;
+	}
+
+	/** @return true if either of *this or Other are invalid or *this is less than or equal to Other. */
+	template<typename ValueType>
+	bool InvalidOrLessOrEqual(const ValueType Other) const
+	{
+		TGuardedSignedInt CheckedOther{ Other };
+		return !ComparisonValid(CheckedOther) || Value <= CheckedOther.Value;
+	}
+
+	/** @return true if either of *this or Other are invalid or *this is greater than Other. */
+	template<typename ValueType>
+	bool InvalidOrGreaterThan(const ValueType Other) const
+	{
+		TGuardedSignedInt CheckedOther{ Other };
+		return !ComparisonValid(CheckedOther) || Value > CheckedOther.Value;
+	}
+
+	/** @return true if either of *this or Other are invalid or *this is greater than or equal to Other. */
+	template<typename ValueType>
+	bool InvalidOrGreaterOrEqual(const ValueType Other) const
+	{
+		TGuardedSignedInt CheckedOther{ Other };
+		return !ComparisonValid(CheckedOther) || Value >= CheckedOther.Value;
+	}
+
+	// Arithmetic operations
+
+	/** @return The negated value. */
+	TGuardedSignedInt operator-() const
+	{
+		// Unary negation (for two's complement) overflows iff the operand is MinValue.
+		return (bIsValid && Value > MinValue) ? TGuardedSignedInt(-Value) : TGuardedSignedInt();
+	}
+
+	/** @return The sum of the two operands. */
+	TGuardedSignedInt operator +(const TGuardedSignedInt Other) const
+	{
+		// Any sum involving an invalid value is invalid.
+		if (!bIsValid || !Other.bIsValid)
+		{
+			return TGuardedSignedInt();
+		}
+
+		// This follows Hacker's Delight, Chapter 2-12
+		// Signed->unsigned conversion and unsigned addition have defined behavior always
+		const UnsignedType UnsignedA = (UnsignedType)Value;
+		const UnsignedType UnsignedB = (UnsignedType)Other.Value;
+		const UnsignedType UnsignedSum = UnsignedA + UnsignedB;
+
+		// Check for signed overflow.
+		// The underlying logic here is pretty simple: if A and B had opposite signs, their sum can't
+		// overflow. If they had the same sign and the sum has the opposite value in the sign bit, we
+		// had an overflow. (See Hacker's Delight Chapter 2-12 for more details.)
+		if ((UnsignedSum ^ UnsignedA) & (UnsignedSum ^ UnsignedB) & UnsignedMSB)
+		{
+			return TGuardedSignedInt();
+		}
+
+		return TGuardedSignedInt(Value + Other.Value);
+	}
+
+	/** @return The difference between the two operands. */
+	TGuardedSignedInt operator -(const TGuardedSignedInt Other) const
+	{
+		// Any difference involving an invalid value is invalid.
+		if (!bIsValid || !Other.bIsValid)
+		{
+			return TGuardedSignedInt();
+		}
+
+		// This follows Hacker's Delight, Chapter 2-12
+		// Signed->unsigned conversion and unsigned subtraction have defined behavior always
+		const UnsignedType UnsignedA = (UnsignedType)Value;
+		const UnsignedType UnsignedB = (UnsignedType)Other.Value;
+		const UnsignedType UnsignedDiff = UnsignedA - UnsignedB;
+
+		// Check for signed overflow.
+		// If A and B have the same sign, the difference can't overflow. Therefore, we test for cases
+		// where the sign bit differs meaning ((UnsignedA ^ UnsignedB) & UnsignedMSB) != 0, and
+		// simultaneously the sign of the difference differs from the sign of the minuend (which should
+		// keep its sign when we're subtracting a value of the opposite sign), meaning
+		// ((UnsignedDiff ^ UnsignedA) & UnsignedMSB) != 0. Combining the two yields:
+		if ((UnsignedA ^ UnsignedB) & (UnsignedDiff ^ UnsignedA) & UnsignedMSB)
+		{
+			return TGuardedSignedInt();
+		}
+
+		return TGuardedSignedInt(Value - Other.Value);
+	}
+
+	/** @return The product of the two operands. */
+	TGuardedSignedInt operator *(const TGuardedSignedInt Other) const
+	{
+		// Any product involving invalid values is invalid.
+		if (!bIsValid || !Other.bIsValid)
+		{
+			return TGuardedSignedInt();
+		}
+
+		// Handle the case where the second factor is 0 specially (why will become clear in a minute).
+		if (Other.Value == 0)
+		{
+			// Anything times 0 is 0.
+			return TGuardedSignedInt(0);
+		}
+
+		// The overflow check is annoying and expensive, but the basic idea is fairly simple:
+		// reduce to an unsigned check of the absolute values. (Again the basic algorithm is
+		// in Hacker's Delight, Chapter 2-12).
+		//
+		// We need the absolute value of the product to be <=MaxValue when the result is positive
+		// (signs of factors same) and <= -MinValue = MaxValue + 1 if the result is negative
+		// (signs of factors opposite).
+		UnsignedType UnsignedA = (UnsignedType)Value;
+		UnsignedType UnsignedB = (UnsignedType)Other.Value;
+		bool bSignsDiffer = false;
+
+		// Determine the unsigned absolute values of A and B carefully (note we can't negate signed
+		// Value or Other.Value, because negating MinValue is UB). We can however subtract their
+		// unsigned values from 0 if the original value was less than zero. While doing this, also
+		// keep track of the sign parity.
+		if (Value < 0)
+		{
+			UnsignedA = UnsignedType(0) - UnsignedA;
+			bSignsDiffer = !bSignsDiffer;
+		}
+
+		if (Other.Value < 0)
+		{
+			UnsignedB = UnsignedType(0) - UnsignedB;
+			bSignsDiffer = !bSignsDiffer;
+		}
+
+		// Determine the unsigned product bound we need based on whether the signs were same or different.
+		const UnsignedType ProductBound = UnsignedType(MaxValue) + (bSignsDiffer ? 1 : 0);
+
+		// We're now in the unsigned case, 0 <= UnsignedA, 0 < UnsignedB (we established b != 0), and for
+		// there not to be overflows we need
+		//   a * b <= ProductBound
+		// <=> a <= ProductBound/b
+		// <=> a <= floor(ProductBound/b)   since a is integer
+		return (UnsignedA <= ProductBound / UnsignedB) ? TGuardedSignedInt(Value * Other.Value) : TGuardedSignedInt();
+	}
+
+	/** @return The quotient when dividing *this by Other. */
+	TGuardedSignedInt operator /(const TGuardedSignedInt Other) const
+	{
+		// Any quotient involving invalid values is invalid.
+		if (!bIsValid || !Other.bIsValid)
+		{
+			return TGuardedSignedInt();
+		}
+
+		// Luckily for us, division generally makes things smaller, so there's only two things to watch
+		// out for: division by zero is not allowed, and division of MinValue by -1 would give -MinValue
+		// which overflows. All other combinations are fine.
+		if (Other.Value == 0 || (Value == MinValue && Other.Value == -1))
+		{
+			return TGuardedSignedInt();
+		}
+
+		return TGuardedSignedInt(Value / Other.Value);
+	}
+
+	/** @return The remainder when dividing *this by Other. */
+	TGuardedSignedInt operator %(const TGuardedSignedInt Other) const
+	{
+		// Any quotient involving invalid values is invalid.
+		if (!bIsValid || !Other.bIsValid)
+		{
+			return TGuardedSignedInt();
+		}
+
+		// Same error cases as for division.
+		if (Other.Value == 0 || (Value == MinValue && Other.Value == -1))
+		{
+			return TGuardedSignedInt();
+		}
+
+		return TGuardedSignedInt(Value % Other.Value);
+	}
+
+	/** @return This value bitwise left-shifted by the operand. */
+	TGuardedSignedInt operator <<(const TGuardedSignedInt Other) const
+	{
+		// Any shift involving invalid values is invalid.
+		if (!bIsValid || !Other.bIsValid)
+		{
+			return TGuardedSignedInt();
+		}
+
+		// Left-shifts by negative values or >= the width of the type are always invalid.
+		if (Other.Value < 0 || Other.Value >= NumBits)
+		{
+			return TGuardedSignedInt();
+		}
+
+		const int ShiftAmount = Other.Value;
+
+		// Once again, taking our overflow-prone expression and using algebra to find
+		// a form that doesn't overflow:
+		//
+		// MinValue <= a * 2^b <= MaxValue
+		// <=> MinValue * 2^(-b) <= a <= MaxValue * 2^(-b)
+		//
+		// The LHS is exact because MinValue is -2^(NumBits - 1) for two's complement,
+		// and we just ensured that 0 <= b < NumBits (with b integer).
+		//
+		// The RHS has a fractional part whereas a is integer; therefore, we can
+		// substitute floor(MaxValue * 2^(-b)) for the RHS without changing the result.
+		//
+		// And that gives us our test!
+		return ((MinValue >> ShiftAmount) <= Value && Value <= (MaxValue >> ShiftAmount)) ? TGuardedSignedInt(Value << ShiftAmount) : TGuardedSignedInt();
+	}
+
+	/** @return This value bitwise right-shifted by the operand. */
+	TGuardedSignedInt operator >>(const TGuardedSignedInt Other) const
+	{
+		// Any shift involving invalid values is invalid.
+		if (!bIsValid || !Other.bIsValid)
+		{
+			return TGuardedSignedInt();
+		}
+
+		// Right-shifts by negative values or >= the width of the type are always invalid.
+		if (Other.Value < 0 || Other.Value >= NumBits)
+		{
+			return TGuardedSignedInt();
+		}
+
+		// Right-shifts don't have any overflow conditions, so we're good!
+		return TGuardedSignedInt(Value >> Other.Value);
+	}
+
+	/** @return The absolute value of *this. */
+	TGuardedSignedInt Abs() const
+	{
+		if (!bIsValid)
+		{
+			return TGuardedSignedInt();
+		}
+
+		// Note the absolute value of MinValue overflows, so this is not completely trivial.
+		// Can't just use TGuardedSignedInt(abs(Value)) here!
+		return (Value < 0) ? -*this : *this;
+	}
+
+	// Mixed-type operators and assignment operators reduce to the base operators systematically
+#define UE_GUARDED_SIGNED_INT_IMPL_BINARY_OPERATOR(OP) \
+	/* Mixed-type expressions that coerce both operands to guarded ints */ \
+	TGuardedSignedInt operator OP(SignedType InB) const { return *this OP TGuardedSignedInt(InB); } \
+	friend TGuardedSignedInt operator OP(SignedType InA, TGuardedSignedInt InB) { return TGuardedSignedInt(InA) OP InB; } \
+	/* Assignment operators, direct and mixed */ \
+	TGuardedSignedInt& operator OP##=(TGuardedSignedInt InB) { return *this = *this OP InB; } \
+	TGuardedSignedInt& operator OP##=(SignedType InB) { return *this = *this OP TGuardedSignedInt(InB); } \
+	/* end */
+
+	UE_GUARDED_SIGNED_INT_IMPL_BINARY_OPERATOR(+)
+	UE_GUARDED_SIGNED_INT_IMPL_BINARY_OPERATOR(-)
+	UE_GUARDED_SIGNED_INT_IMPL_BINARY_OPERATOR(*)
+	UE_GUARDED_SIGNED_INT_IMPL_BINARY_OPERATOR(/)
+	UE_GUARDED_SIGNED_INT_IMPL_BINARY_OPERATOR(%)
+	UE_GUARDED_SIGNED_INT_IMPL_BINARY_OPERATOR(<<)
+	UE_GUARDED_SIGNED_INT_IMPL_BINARY_OPERATOR(>>)
+
+#undef UE_GUARDED_SIGNED_INT_IMPL_BINARY_OPERATOR
+};
+
+/** Guarded 32-bit integer class. Used to deal with integer data from untrusted sources in size computations etc. */
+using FGuardedInt32 = TGuardedSignedInt<int32>;
+
+/** Guarded 64-bit integer class. Used to deal with integer data from untrusted sources in size computations etc. */
+using FGuardedInt64 = TGuardedSignedInt<int64>;
