Fix GuardedInt compile bug in GetChecked

diff --git a/Engine/Source/Runtime/ImageCore/Private/ImageCoreTest.cpp b/Engine/Source/Runtime/ImageCore/Private/ImageCoreTest.cpp
--- a/Engine/Source/Runtime/ImageCore/Private/ImageCoreTest.cpp
+++ b/Engine/Source/Runtime/ImageCore/Private/ImageCoreTest.cpp
@@ -61,6 +61,15 @@ bool FGuardedIntTest::RunTest(const FString& Parameters)
 			FGuardedInt8 CheckedA{ A };
 			FGuardedInt8 CheckedB{ B };
 
+			// This doesn't test much by itself but ensures we have at least compiled GetChecked.
+			if (A <= 127 && B == 0)
+			{
+				if (CheckedA.GetChecked() != A)
+				{
+					AddError(FString::Printf(TEXT("GetChecked A=%d"), A));
+				}
+			}
+
 			// Compute reference results for the operations that are not defined everywhere
 			int32 ADivB = -256;
 			int32 AModB = -256;
diff --git a/Engine/Source/Runtime/ImageCore/Public/GuardedInt.h b/Engine/Source/Runtime/ImageCore/Public/GuardedInt.h
--- a/Engine/Source/Runtime/ImageCore/Public/GuardedInt.h
+++ b/Engine/Source/Runtime/ImageCore/Public/GuardedInt.h
@@ -3,6 +3,7 @@
 #pragma once
 
 #include "CoreTypes.h"
+#include "Misc/AssertionMacros.h"
 #include "Math/NumericLimits.h"
 #include <type_traits>
 
@@ -104,7 +105,7 @@ public:
 	 */
 	const SignedType GetChecked(const SignedType DefaultValue = 0) const
 	{
-		checkf(IsValid(), "Invalid value in TGuardedSignedInt::GetChecked.");
+		checkf(IsValid(), TEXT("Invalid value in TGuardedSignedInt::GetChecked."));
 		return Get(DefaultValue);
 	}
 
