Fix for incorrect view size for RGBA16

diff --git a/Engine/Source/Runtime/ImageCore/Public/ImageCore.h b/Engine/Source/Runtime/ImageCore/Public/ImageCore.h
--- a/Engine/Source/Runtime/ImageCore/Public/ImageCore.h
+++ b/Engine/Source/Runtime/ImageCore/Public/ImageCore.h
@@ -312,7 +312,7 @@ public:
 	TArrayView64<uint16> AsRGBA16() const
 	{
 		check(Format == ERawImageFormat::RGBA16);
-		return { (uint16*)RawData, GetNumPixels() };
+		return { (uint16*)RawData, GetNumPixels() * 4 };
 	}
 
 	TArrayView64<FFloat16Color> AsRGBA16F() const
