ImageCore.cpp : align pixel jobs to cache lines

diff --git a/Engine/Source/Runtime/ImageCore/Private/ImageCore.cpp b/Engine/Source/Runtime/ImageCore/Private/ImageCore.cpp
--- a/Engine/Source/Runtime/ImageCore/Private/ImageCore.cpp
+++ b/Engine/Source/Runtime/ImageCore/Private/ImageCore.cpp
@@ -63,10 +63,22 @@ static inline int32 ParallelForComputeNumJobs(int64 & OutNumItemsPerJob,int64 Nu
 static constexpr int64 MinPixelsPerJob = 16384;
 // Surfaces of VT tile size or smaller will not parallelize at all :
 static constexpr int64 MinPixelsForAnyJob = 136*136;
+// Jobs split on pixel count (not rows) are aligned to PixelsPerJobAlignment :
+static constexpr int64 PixelsPerJobAlignment = 64; // must be power of 2
 
 IMAGECORE_API int32 ImageParallelForComputeNumJobsForPixels(int64 & OutNumPixelsPerJob,int64 NumPixels)
 {
-	return ParallelForComputeNumJobs(OutNumPixelsPerJob,NumPixels,MinPixelsPerJob,MinPixelsForAnyJob);
+	int32 NumJobs = ParallelForComputeNumJobs(OutNumPixelsPerJob,NumPixels,MinPixelsPerJob,MinPixelsForAnyJob);
+	
+	if ( NumJobs > 1 )
+	{
+		// align up to PixelsPerJobAlignment :
+		//	(PixelsPerJobAlignment should be power of 2)
+		OutNumPixelsPerJob = (OutNumPixelsPerJob + PixelsPerJobAlignment-1) & (~ (PixelsPerJobAlignment-1) );
+		// recompute NumJobs :
+		NumJobs = (NumPixels + OutNumPixelsPerJob-1) / OutNumPixelsPerJob;
+	}
+	return NumJobs;
 }
 
 IMAGECORE_API int32 ImageParallelForComputeNumJobsForRows(int32 & OutNumItemsPerJob,int32 SizeX,int32 SizeY)
