ImageCore: Extracting the working color space transformation method into the public namespace.

#rnx

diff --git a/Engine/Source/Developer/TextureCompressor/Private/TextureCompressorModule.cpp b/Engine/Source/Developer/TextureCompressor/Private/TextureCompressorModule.cpp
--- a/Engine/Source/Developer/TextureCompressor/Private/TextureCompressorModule.cpp
+++ b/Engine/Source/Developer/TextureCompressor/Private/TextureCompressorModule.cpp
@@ -3934,7 +3934,8 @@ private:
 
 			if (BuildSettings.bHasColorSpaceDefinition)
 			{
-				Mip->TransformToWorkingColorSpace(
+				FImageCore::TransformToWorkingColorSpace(
+					*Mip,
 					FVector2d(BuildSettings.RedChromaticityCoordinate),
 					FVector2d(BuildSettings.GreenChromaticityCoordinate),
 					FVector2d(BuildSettings.BlueChromaticityCoordinate),
diff --git a/Engine/Source/Runtime/ImageCore/Private/ImageCore.cpp b/Engine/Source/Runtime/ImageCore/Private/ImageCore.cpp
--- a/Engine/Source/Runtime/ImageCore/Private/ImageCore.cpp
+++ b/Engine/Source/Runtime/ImageCore/Private/ImageCore.cpp
@@ -741,33 +741,7 @@ FORCEINLINE static FLinearColor SaturateToHalfFloat(const FLinearColor& LinearCo
 
 void FImage::TransformToWorkingColorSpace(const FVector2d& SourceRedChromaticity, const FVector2d& SourceGreenChromaticity, const FVector2d& SourceBlueChromaticity, const FVector2d& SourceWhiteChromaticity, UE::Color::EChromaticAdaptationMethod Method, double EqualityTolerance)
 {
-	TRACE_CPUPROFILER_EVENT_SCOPE(Texture.TransformToWorkingColorSpace);
-
-	check(GammaSpace == EGammaSpace::Linear);
-
-	const UE::Color::FColorSpace Source(SourceRedChromaticity, SourceGreenChromaticity, SourceBlueChromaticity, SourceWhiteChromaticity);
-	const UE::Color::FColorSpace& Target = UE::Color::FColorSpace::GetWorking();
-
-	if (Source.Equals(Target, EqualityTolerance))
-	{
-		UE_LOG(LogImageCore, VeryVerbose, TEXT("Source and working color spaces are equal within tolerance, bypass color space transformation."));
-		return;
-	}
-
-	UE::Color::FColorSpaceTransform Transform(Source, Target, Method);
-
-	FLinearColor* ImageColors = AsRGBA32F().GetData();
-
-	const int64 NumTexels = int64(SizeX) * SizeY * NumSlices;
-	int64 TexelsPerJob;
-	int32 NumJobs = ImageParallelForComputeNumJobsForPixels(TexelsPerJob,NumTexels);
-
-	ParallelLoop(TEXT("Texture.TransformToWorkingColorSpace.PF"), NumJobs, TexelsPerJob, NumTexels, [Transform, ImageColors](int64 TexelIndex)
-	{
-		FLinearColor Color = ImageColors[TexelIndex];
-		Color = Transform.Apply(Color);
-		ImageColors[TexelIndex] = SaturateToHalfFloat(Color);
-	});
+	return FImageCore::TransformToWorkingColorSpace(*this, SourceRedChromaticity, SourceGreenChromaticity, SourceBlueChromaticity, SourceWhiteChromaticity, Method, EqualityTolerance);
 }
 
 static FLinearColor SampleImage(const FLinearColor* Pixels, int Width, int Height, float X, float Y)
@@ -1552,3 +1526,34 @@ void FImageCore::ComputeChannelLinearMinMax(const FImage & InImage, FLinearColor
 	VectorStore(VMin,&OutMin.Component(0));
 	VectorStore(VMax,&OutMax.Component(0));
 }
+
+void FImageCore::TransformToWorkingColorSpace(const FImageView& InLinearImage, const FVector2d& SourceRedChromaticity, const FVector2d& SourceGreenChromaticity, const FVector2d& SourceBlueChromaticity, const FVector2d& SourceWhiteChromaticity, UE::Color::EChromaticAdaptationMethod Method, double EqualityTolerance)
+{
+	TRACE_CPUPROFILER_EVENT_SCOPE(Texture.TransformToWorkingColorSpace);
+
+	check(InLinearImage.GammaSpace == EGammaSpace::Linear);
+
+	const UE::Color::FColorSpace Source(SourceRedChromaticity, SourceGreenChromaticity, SourceBlueChromaticity, SourceWhiteChromaticity);
+	const UE::Color::FColorSpace& Target = UE::Color::FColorSpace::GetWorking();
+
+	if (Source.Equals(Target, EqualityTolerance))
+	{
+		UE_LOG(LogImageCore, VeryVerbose, TEXT("Source and working color spaces are equal within tolerance, bypass color space transformation."));
+		return;
+	}
+
+	UE::Color::FColorSpaceTransform Transform(Source, Target, Method);
+
+	FLinearColor* ImageColors = InLinearImage.AsRGBA32F().GetData();
+
+	const int64 NumTexels = int64(InLinearImage.SizeX) * InLinearImage.SizeY * InLinearImage.NumSlices;
+	int64 TexelsPerJob;
+	int32 NumJobs = ImageParallelForComputeNumJobsForPixels(TexelsPerJob, NumTexels);
+
+	ParallelLoop(TEXT("Texture.TransformToWorkingColorSpace.PF"), NumJobs, TexelsPerJob, NumTexels, [Transform, ImageColors](int64 TexelIndex)
+		{
+			FLinearColor Color = ImageColors[TexelIndex];
+			Color = Transform.Apply(Color);
+			ImageColors[TexelIndex] = SaturateToHalfFloat(Color);
+		});
+}
diff --git a/Engine/Source/Runtime/ImageCore/Public/ImageCore.h b/Engine/Source/Runtime/ImageCore/Public/ImageCore.h
--- a/Engine/Source/Runtime/ImageCore/Public/ImageCore.h
+++ b/Engine/Source/Runtime/ImageCore/Public/ImageCore.h
@@ -459,16 +459,7 @@ public:
 		Linearize(0,DestImage);
 	}
 
-	/**
-	 * Apply a color space transformation from the source chromaticities to the engine's working color space.
-	 *
-	 * @param SourceRedChromaticity - The red chromaticity coordinate of the source color space.
-	 * @param SourceGreenChromaticity - The green chromaticity coordinate of the source color space.
-	 * @param SourceBlueChromaticity - The blue chromaticity coordinate of the source color space.
-	 * @param SourceWhiteChromaticity - The white chromaticity coordinate of the source color space.
-	 * @param Method - The chromatic adapation method.
-	 * @param EqualityTolerance - The tolerance for the source and working color space chromaticities to be considered equal, bypassing the transform.
-	 */
+	UE_DEPRECATED(5.3, "TransformToWorkingColorSpace is deprecated, please use the function in FImageCore.")
 	IMAGECORE_API void TransformToWorkingColorSpace(const FVector2d& SourceRedChromaticity, const FVector2d& SourceGreenChromaticity, const FVector2d& SourceBlueChromaticity, const FVector2d& SourceWhiteChromaticity, UE::Color::EChromaticAdaptationMethod Method, double EqualityTolerance = 1.e-7);
 
 	/**
@@ -714,12 +705,17 @@ IMAGECORE_API void ResizeTo(const FImageView & SourceImage,FImage& DestImage, in
  */
 IMAGECORE_API void ComputeChannelLinearMinMax(const FImage & InImage, FLinearColor & OutMin, FLinearColor & OutMax);
 
-};
-
+/**
+ * Apply a color space transformation from the source chromaticities to the engine's working color space.
+ *
+ * @param InLinearImage - The image to convert, which must be linear.
+ * @param SourceRedChromaticity - The red chromaticity coordinate of the source color space.
+ * @param SourceGreenChromaticity - The green chromaticity coordinate of the source color space.
+ * @param SourceBlueChromaticity - The blue chromaticity coordinate of the source color space.
+ * @param SourceWhiteChromaticity - The white chromaticity coordinate of the source color space.
+ * @param Method - The chromatic adapation method.
+ * @param EqualityTolerance - The tolerance for the source and working color space chromaticities to be considered equal, bypassing the transform.
+ */
+IMAGECORE_API void TransformToWorkingColorSpace(const FImageView& InLinearImage, const FVector2d& SourceRedChromaticity, const FVector2d& SourceGreenChromaticity, const FVector2d& SourceBlueChromaticity, const FVector2d& SourceWhiteChromaticity, UE::Color::EChromaticAdaptationMethod Method, double EqualityTolerance = 1.e-7);
 
-// non-namespaced CopyImage for backwards compatibility
-UE_DEPRECATED(5.1,"Use FImageCore::CopyImage")
-FORCEINLINE void CopyImage(const FImageView & SrcImage,const FImageView & DestImage)
-{
-	return FImageCore::CopyImage(SrcImage,DestImage);
-}
\ No newline at end of file
+};
