ImageCore parallel DetectAlpha

diff --git a/Engine/Source/Runtime/ImageCore/Private/ImageCore.cpp b/Engine/Source/Runtime/ImageCore/Private/ImageCore.cpp
--- a/Engine/Source/Runtime/ImageCore/Private/ImageCore.cpp
+++ b/Engine/Source/Runtime/ImageCore/Private/ImageCore.cpp
@@ -104,6 +104,34 @@ static void ParallelLoop(const TCHAR* DebugName, int32 NumJobs, int64 TexelsPerJ
 	});
 }
 
+// ParallelOr : call Func() on all texels ; returns true if Func is true for any texel
+template <typename Lambda>
+static bool ParallelOr(const TCHAR* DebugName, int32 NumJobs, int64 TexelsPerJob, int64 NumTexels, const Lambda& Func)
+{
+	// merged_bool starts false, is set to true if any thread sees a true
+	std::atomic<bool> merged_bool(false);
+
+	ParallelFor(DebugName, NumJobs, 1, [&](int64 JobIndex)
+	{
+		const int64 StartIndex = JobIndex * TexelsPerJob;
+		const int64 EndIndex = FMath::Min(StartIndex + TexelsPerJob, NumTexels);
+		for (int64 TexelIndex = StartIndex; TexelIndex < EndIndex; ++TexelIndex)
+		{
+			if ( Func(TexelIndex) )
+			{
+				merged_bool.store(true,std::memory_order_release);
+				return;
+			}
+		}
+		// if Func was really slow, you might want to check if merged_bool was set to true on some other thread
+		//	and terminate our loop when that is true
+		//  but that causes cache traffic for the shared variable, so should not be done in tight loops
+	});
+
+	bool ret = merged_bool.load(std::memory_order_acquire);
+	return ret;
+}
+
 template <int64 TexelsPerVec, typename LambdaVec, typename Lambda>
 static void ParallelLoop(const TCHAR* DebugName, int32 NumJobs, int64 TexelsPerJob, int64 NumTexels, const LambdaVec& FuncVec, const Lambda& Func)
 {
@@ -1320,7 +1348,6 @@ void FImageCore::SanitizeFloat16AndSetAlphaOpaqueForBC6H(const FImageView & InOu
 	}
 }
 
-
 bool FImageCore::DetectAlphaChannel(const FImageView & InImage)
 {
 	TRACE_CPUPROFILER_EVENT_SCOPE(Texture.DetectAlphaChannel);
@@ -1334,75 +1361,61 @@ bool FImageCore::DetectAlphaChannel(const FImageView & InImage)
 	//  images with only alpha larger than this are treated as opaque
 	const float FloatNonOpaqueAlpha = 254.5f / 255.f; // the U8 alpha threshold
 
-	int64 NumPixels = (int64)InImage.SizeX * InImage.SizeY * InImage.NumSlices;
+	int64 NumTexels = InImage.GetNumPixels();
+	int64 TexelsPerJob;
+	int32 NumJobs = ImageParallelForComputeNumJobsForPixels(TexelsPerJob,NumTexels);
 
 	if (InImage.Format == ERawImageFormat::BGRA8)
 	{
 		TArrayView64<const FColor> SrcColorArray = InImage.AsBGRA8();
-		check(SrcColorArray.Num() == NumPixels);
+		check(SrcColorArray.Num() == NumTexels);
 
 		const FColor* ColorPtr = &SrcColorArray[0];
-		const FColor* EndPtr = ColorPtr + SrcColorArray.Num();
-
-		for (; ColorPtr < EndPtr; ++ColorPtr)
+		
+		return ParallelOr(TEXT("Texture.DetectAlphaChannel.PF"), NumJobs, TexelsPerJob, NumTexels, [&](int64 TexelIndex)
 		{
-			if (ColorPtr->A != 255)
-			{
-				return true;
-			}
-		}
+			return (ColorPtr[TexelIndex].A != 255);
+		});
 	}
 	else if (InImage.Format == ERawImageFormat::RGBA32F)
 	{
 		TArrayView64<const FLinearColor> SrcColorArray = InImage.AsRGBA32F();
-		check(SrcColorArray.Num() == NumPixels);
+		check(SrcColorArray.Num() == NumTexels);
 
 		const FLinearColor* ColorPtr = &SrcColorArray[0];
-		const FLinearColor* EndPtr = ColorPtr + SrcColorArray.Num();
-
-		for (; ColorPtr < EndPtr; ++ColorPtr)
+		
+		return ParallelOr(TEXT("Texture.DetectAlphaChannel.PF"), NumJobs, TexelsPerJob, NumTexels, [&](int64 TexelIndex)
 		{
-			if (ColorPtr->A <= FloatNonOpaqueAlpha)
-			{
-				return true;
-			}
-		}
+			return (ColorPtr[TexelIndex].A <= FloatNonOpaqueAlpha);
+		});
 	}
 	else if (InImage.Format == ERawImageFormat::RGBA16)
 	{
 		TArrayView64<const uint16> SrcChannelArray = InImage.AsRGBA16();
-		check(SrcChannelArray.Num() == NumPixels * 4);
+		check(SrcChannelArray.Num() == NumTexels * 4);
 
 		const uint16* ChannelPtr = &SrcChannelArray[0];
-		const uint16* EndPtr = ChannelPtr + SrcChannelArray.Num();
-
-		for (; ChannelPtr < EndPtr; ChannelPtr += 4)
+		
+		return ParallelOr(TEXT("Texture.DetectAlphaChannel.PF"), NumJobs, TexelsPerJob, NumTexels, [&](int64 TexelIndex)
 		{
-			if (ChannelPtr[3] != 0xFFFF)
-			{
-				return true;
-			}
-		}
+			return (ChannelPtr[TexelIndex*4+3] != 0xFFFF);
+		});
 	}
 	else if (InImage.Format == ERawImageFormat::RGBA16F)
 	{
 		TArrayView64<const FFloat16Color> SrcColorArray = InImage.AsRGBA16F();
-		check(SrcColorArray.Num() == NumPixels);
+		check(SrcColorArray.Num() == NumTexels);
 
 		const FFloat16Color* ColorPtr = &SrcColorArray[0];
-		const FFloat16Color* EndPtr = ColorPtr + SrcColorArray.Num();
+		
+		// 16F closest to 1.0 is 0.99951172
+		// use the float tolerance here? or check exactly ?
+		// use the same FloatNonOpaqueAlpha tolerance for consistency ?
 
-		for (; ColorPtr < EndPtr; ++ColorPtr)
+		return ParallelOr(TEXT("Texture.DetectAlphaChannel.PF"), NumJobs, TexelsPerJob, NumTexels, [&](int64 TexelIndex)
 		{
-			// 16F closest to 1.0 is 0.99951172
-			// use the float tolerance here? or check exactly ?
-			//if ( ColorPtr->A.GetFloat() < 1.f )
-			// use the same FloatNonOpaqueAlpha tolerance for consistency ?
-			if (ColorPtr->A.GetFloat() < FloatNonOpaqueAlpha)
-			{
-				return true;
-			}
-		}
+			return (ColorPtr[TexelIndex].A.GetFloat() <= FloatNonOpaqueAlpha);
+		});
 	}
 	else if (InImage.Format == ERawImageFormat::G8 ||
 		InImage.Format == ERawImageFormat::BGRE8 ||
@@ -1421,7 +1434,6 @@ bool FImageCore::DetectAlphaChannel(const FImageView & InImage)
 	return false;
 }
 
-
 void FImageCore::SetAlphaOpaque(const FImageView & InImage)
 {
 	int64 NumPixels = (int64)InImage.SizeX * InImage.SizeY * InImage.NumSlices;
