nop comments only about preferring FImage usage

#rnx

diff --git a/Engine/Source/Runtime/Engine/Public/ImageUtils.h b/Engine/Source/Runtime/Engine/Public/ImageUtils.h
--- a/Engine/Source/Runtime/Engine/Public/ImageUtils.h
+++ b/Engine/Source/Runtime/Engine/Public/ImageUtils.h
@@ -11,6 +11,19 @@ ImageUtils.h: Image/Textures utility functions
 #include "ImageCore.h"
 #include "Engine/Texture.h"
 
+/*
+
+NOTE: Please prefer to work with images through ImageCore FImage/FImageView
+(not TextureSource/TextureSourceFormat or raw arrays of bytes)
+
+Use FImageUtils::CompressImage/DecompressImage
+
+The easiest way to load/save an image is FImageUtils::LoadImage/SaveImageByExtension
+
+Typically let the user's filename choose the image format, don't hard-code them in code.
+
+*/
+
 class UTexture2D;
 class UTexture2DArray;
 class UTextureCube;
diff --git a/Engine/Source/Runtime/ImageCore/Private/ImageCore.cpp b/Engine/Source/Runtime/ImageCore/Private/ImageCore.cpp
--- a/Engine/Source/Runtime/ImageCore/Private/ImageCore.cpp
+++ b/Engine/Source/Runtime/ImageCore/Private/ImageCore.cpp
@@ -1603,7 +1603,7 @@ void FImageCore::SetAlphaOpaque(const FImageView & InImage)
 }
 
 
-void FImageCore::ComputeChannelLinearMinMax(const FImage & InImage, FLinearColor & OutMin, FLinearColor & OutMax)
+void FImageCore::ComputeChannelLinearMinMax(const FImageView & InImage, FLinearColor & OutMin, FLinearColor & OutMax)
 {
 	// @todo Oodle : for speed, we should ideally scan the image for min/max in its native pixel format
 	//	then only convert the min/max colors to float linear after the scan
diff --git a/Engine/Source/Runtime/ImageCore/Public/ImageCore.h b/Engine/Source/Runtime/ImageCore/Public/ImageCore.h
--- a/Engine/Source/Runtime/ImageCore/Public/ImageCore.h
+++ b/Engine/Source/Runtime/ImageCore/Public/ImageCore.h
@@ -20,7 +20,11 @@
 * note to authors: as much as possible, write functions that act on FImageView
 *  for example reading and modifying pixels? use FImageView
 * use FImage when you may need to change the format or allocate a new image.
-* 
+*
+* prefer using FImage/FImageView instead of TextureSourceFormat or raw arrays or bytes
+*
+* Try not to write code that switches on pixel format, as formats may be added and it creates a fragile maintenance problem.
+*
 */
 
 
@@ -726,7 +730,7 @@ IMAGECORE_API void ResizeTo(const FImageView & SourceImage,FImage& DestImage, in
  * @param OutMin - filled with the minimum of the color channels
  * @param OutMax - filled with the maximum of the color channels
  */
-IMAGECORE_API void ComputeChannelLinearMinMax(const FImage & InImage, FLinearColor & OutMin, FLinearColor & OutMax);
+IMAGECORE_API void ComputeChannelLinearMinMax(const FImageView & InImage, FLinearColor & OutMin, FLinearColor & OutMax);
 
 /**
  * Apply a color space transformation from the source chromaticities to the engine's working color space.
diff --git a/Engine/Source/Runtime/ImageWrapper/Public/IImageWrapper.h b/Engine/Source/Runtime/ImageWrapper/Public/IImageWrapper.h
--- a/Engine/Source/Runtime/ImageWrapper/Public/IImageWrapper.h
+++ b/Engine/Source/Runtime/ImageWrapper/Public/IImageWrapper.h
@@ -8,6 +8,17 @@
 #include "Templates/SharedPointer.h"
 #include "ImageCore.h"
 
+/**
+
+NOTE: you should not write code that talks directly to individual ImageWrappers
+
+Instead use ImageWrapperModule CompressImage/DecompressImage
+
+Prefer the new interface that go through FImage not TArray of bytes
+
+
+**/
+
 /**
  * Enumerates the types of image formats this class can handle.
  */
diff --git a/Engine/Source/Runtime/ImageWrapper/Public/IImageWrapperModule.h b/Engine/Source/Runtime/ImageWrapper/Public/IImageWrapperModule.h
--- a/Engine/Source/Runtime/ImageWrapper/Public/IImageWrapperModule.h
+++ b/Engine/Source/Runtime/ImageWrapper/Public/IImageWrapperModule.h
@@ -12,6 +12,13 @@
 /**
  * Interface for image wrapper modules.
  * 
+ * If you have "Engine" module, then just use FImageUtils::CompressImage/DecompressImage
+ *
+ * The easiest way to load/save an image is FImageUtils::LoadImage/SaveImageByExtension
+ *
+ * NOTE: Please prefer to work with images through ImageCore FImage/FImageView 
+ * (not TextureSource/TextureSourceFormat or raw arrays of bytes)
+ * 
  * 
  * note on SRGB/Gamma handling :
  * it is assumed that non-U8 data is always Linear
