ImageCore: GetOnePixelLinear helper cleanup
GetOnePixelLinear should take EGammaSpace not bSRGB
GetOnePixelLinear helpers on the FImage/FImageView classes

#rnx

diff --git a/Engine/Source/Runtime/ImageCore/Private/ImageCore.cpp b/Engine/Source/Runtime/ImageCore/Private/ImageCore.cpp
--- a/Engine/Source/Runtime/ImageCore/Private/ImageCore.cpp
+++ b/Engine/Source/Runtime/ImageCore/Private/ImageCore.cpp
@@ -1355,7 +1355,7 @@ IMAGECORE_API bool ERawImageFormat::IsHDR(Type Format)
 	return Format == RGBA16F || Format == RGBA32F || Format == R16F || Format == R32F || Format == BGRE8;
 }
 
-IMAGECORE_API FLinearColor ERawImageFormat::GetOnePixelLinear(const void * PixelData,Type Format,bool bSRGB)
+IMAGECORE_API const FLinearColor ERawImageFormat::GetOnePixelLinear(const void * PixelData,ERawImageFormat::Type Format,EGammaSpace Gamma)
 {
 	switch(Format)
 	{
@@ -1363,22 +1363,26 @@ IMAGECORE_API FLinearColor ERawImageFormat::GetOnePixelLinear(const void * Pixel
 	{
 		uint8 Gray = ((const uint8 *)PixelData)[0];
 		FColor Color(Gray,Gray,Gray);
-		if ( bSRGB )
+		if ( Gamma == EGammaSpace::sRGB )
 			return FLinearColor::FromSRGBColor(Color);
+		else if ( Gamma == EGammaSpace::Pow22 )
+			return FLinearColor::FromPow22Color(Color);
 		else
 			return Color.ReinterpretAsLinear();
 	}
 	case G16:
 	{
 		const uint16 * Samples = (const uint16 *)PixelData;
-		float Gray = Samples[0]/65535.f;
+		float Gray = Samples[0]*(1.f/65535.f);
 		return FLinearColor(Gray,Gray,Gray,1.f);
 	}
 	case BGRA8:
 	{
 		FColor Color = *((const FColor *)PixelData);
-		if ( bSRGB )
+		if ( Gamma == EGammaSpace::sRGB )
 			return FLinearColor::FromSRGBColor(Color);
+		else if ( Gamma == EGammaSpace::Pow22 )
+			return FLinearColor::FromPow22Color(Color);
 		else
 			return Color.ReinterpretAsLinear();
 	}
@@ -1391,10 +1395,10 @@ IMAGECORE_API FLinearColor ERawImageFormat::GetOnePixelLinear(const void * Pixel
 	{
 		const uint16 * Samples = (const uint16 *)PixelData;
 		return FLinearColor(
-			Samples[0]/65535.f, 
-			Samples[1]/65535.f, 
-			Samples[2]/65535.f, 
-			Samples[3]/65535.f);
+			Samples[0]*(1.f/65535.f), 
+			Samples[1]*(1.f/65535.f), 
+			Samples[2]*(1.f/65535.f), 
+			Samples[3]*(1.f/65535.f));
 	}
 	case RGBA16F:
 	{
diff --git a/Engine/Source/Runtime/ImageCore/Public/ImageCore.h b/Engine/Source/Runtime/ImageCore/Public/ImageCore.h
--- a/Engine/Source/Runtime/ImageCore/Public/ImageCore.h
+++ b/Engine/Source/Runtime/ImageCore/Public/ImageCore.h
@@ -77,8 +77,14 @@ namespace ERawImageFormat
 	IMAGECORE_API bool IsHDR(Type Format);
 	
 	// Get one pixel of Format type from PixelData and return in Linear color
-	IMAGECORE_API FLinearColor GetOnePixelLinear(const void * PixelData,Type Format,bool bSRGB);	
+	IMAGECORE_API const FLinearColor GetOnePixelLinear(const void * PixelData,Type Format,EGammaSpace Gamma);
 	
+	// Get one pixel of Format type from PixelData and return in Linear color
+	FORCEINLINE const FLinearColor GetOnePixelLinear(const void * PixelData,Type Format,bool bSRGB)
+	{
+		return GetOnePixelLinear(PixelData,Format,bSRGB ? EGammaSpace::sRGB : EGammaSpace::Linear);
+	}
+
 	// G8 and BGRA8 are affected by Gamma
 	//	16/32 is NOT
 	FORCEINLINE bool GetFormatNeedsGammaSpace(Type Format)
@@ -214,6 +220,23 @@ struct FImageInfo
 
 		return GammaSpace;
 	}
+	
+	// get offset of a pixel from the base pointer, in bytes
+	FORCEINLINE int64 GetPixelOffsetBytes(int32 X,int32 Y,int32 Slice = 0) const
+	{
+		checkSlow( X >= 0 && X < SizeX );
+		checkSlow( Y >= 0 && Y < SizeY );
+		checkSlow( Slice >= 0 && Slice < NumSlices );
+
+		int64 Offset = Slice * GetSliceNumPixels();
+		Offset += Y * (int64)SizeX;
+		Offset += X;
+		// Offset is now is pixels
+		Offset *= GetBytesPerPixel();
+
+		return Offset;
+	}
+
 };
 
 /***
@@ -298,6 +321,21 @@ struct FImageView : public FImageInfo
 		CopyTo(DestImage,Format,GammaSpace);
 	}
 	
+	// get a pointer to a pixel
+	FORCEINLINE void * GetPixelPointer(int32 X,int32 Y,int32 Slice=0) const
+	{
+		uint8 * Ptr = (uint8 *)RawData;
+		Ptr += GetPixelOffsetBytes(X,Y,Slice);
+		return (void *)Ptr;
+	}
+
+	// Get one pixel from the image and return in Linear color
+	IMAGECORE_API const FLinearColor GetOnePixelLinear(int32 X,int32 Y,int32 Slice=0) const
+	{
+		void * Ptr = GetPixelPointer(X,Y,Slice);
+		return ERawImageFormat::GetOnePixelLinear(Ptr,Format,GammaSpace);
+	}
+
 public:
 
 	// Convenience accessors to raw data
@@ -520,6 +558,20 @@ public:
 	 * @param Info - The image description.  Can also pass an FImage or FImageView to this function.
 	 */
 	IMAGECORE_API void Init(const FImageInfo & Info);
+	
+	// get a pointer to a pixel
+	FORCEINLINE void * GetPixelPointer(int32 X,int32 Y,int32 Slice=0) const
+	{
+		int64 Offset = GetPixelOffsetBytes(X,Y,Slice);
+		return (void *)&RawData[Offset];
+	}
+
+	// Get one pixel from the image and return in Linear color
+	IMAGECORE_API const FLinearColor GetOnePixelLinear(int32 X,int32 Y,int32 Slice=0) const
+	{
+		void * Ptr = GetPixelPointer(X,Y,Slice);
+		return ERawImageFormat::GetOnePixelLinear(Ptr,Format,GammaSpace);
+	}
 
 public:
 
