ImageCore.build.cs clarifying comment

#rnx

diff --git a/Engine/Source/Runtime/ImageCore/ImageCore.Build.cs b/Engine/Source/Runtime/ImageCore/ImageCore.Build.cs
--- a/Engine/Source/Runtime/ImageCore/ImageCore.Build.cs
+++ b/Engine/Source/Runtime/ImageCore/ImageCore.Build.cs
@@ -7,7 +7,9 @@ public class ImageCore : ModuleRules
 	public ImageCore(ReadOnlyTargetRules Target) : base(Target)
 	{
 		// include only, no link :
-		// for TextureDefines.h :
+		// ImageCoreUtils.h includes TextureDefines.h from Engine, so technically it should be a "PublicIncludePathModuleNames" include,
+		// but that breaks the build because there are other modules that use ImageCore that have conflicts with Engine headers.
+		// So, leave it as PrivateIncludePathModuleNames, but then any module that uses ImageCoreUtils needs to PrivateIncludePathModuleNames Engine.
 		// PublicIncludePathModuleNames
 		PrivateIncludePathModuleNames.AddRange(new string[]
 		{
