ImageParallelFor for easy parallel work on pixels
ImageParallelProcessLinearPixels
ComputeImageLinearAverage
Make Bokeh computation use the new routines as a test

#rnx

diff --git a/Engine/Source/Developer/TextureCompressor/Private/TextureCompressorModule.cpp b/Engine/Source/Developer/TextureCompressor/Private/TextureCompressorModule.cpp
--- a/Engine/Source/Developer/TextureCompressor/Private/TextureCompressorModule.cpp
+++ b/Engine/Source/Developer/TextureCompressor/Private/TextureCompressorModule.cpp
@@ -17,6 +17,7 @@
 #include "Misc/Paths.h"
 #include "Tasks/Task.h"
 #include "ImageCore.h"
+#include "ImageParallelFor.h"
 #include <cmath>
 
 #if PLATFORM_WINDOWS
@@ -2787,19 +2788,12 @@ void ITextureCompressorModule::AdjustImageColors(FImage& Image, const FTextureBu
 	{
 		TRACE_CPUPROFILER_EVENT_SCOPE(Texture.AdjustImageColors);
 
-		const int64 NumPixels = (int64)Image.SizeX * Image.SizeY * Image.NumSlices;
-		TArrayView64<FLinearColor> ImageColors = Image.AsRGBA32F();
-
-		int64 NumPixelsEachJob;
-		int32 NumJobs = ImageParallelForComputeNumJobsForPixels(NumPixelsEachJob,NumPixels);
-
-		ParallelFor( TEXT("Texture.AdjustImageColorsFunc.PF"),NumJobs,1, [&](int32 Index)
+		FImageCore::ImageParallelFor( TEXT("Texture.AdjustImageColorsFunc.PF"),Image, [&](FImageView & ImagePart)
 		{
-			int64 StartIndex = Index * NumPixelsEachJob;
-			int64 EndIndex = FMath::Min(StartIndex + NumPixelsEachJob, NumPixels);
+			TArrayView64<FLinearColor> ImageColors = ImagePart.AsRGBA32F();
 
-			FLinearColor * First = &ImageColors[StartIndex];
-			int64 Count = EndIndex-StartIndex;
+			FLinearColor * First = &ImageColors[0];
+			int64 Count = ImageColors.Num();
 
 			// Use new AdjustColors code only for newly added textures
 			// So existing textures maintain exactly same output as before
@@ -2811,7 +2805,7 @@ void ITextureCompressorModule::AdjustImageColors(FImage& Image, const FTextureBu
 			{
 				AdjustColorsOld(First, Count, InBuildSettings);
 			}
-		}, EParallelForFlags::Unbalanced);
+		});
 	}
 }
 
@@ -2824,25 +2818,14 @@ static void ComputeBokehAlpha(FImage& Image)
 {
 	check( Image.SizeX > 0 && Image.SizeY > 0 );
 
-	const int64 NumPixels = Image.GetNumPixels();
-	TArrayView64<FLinearColor> ImageColors = Image.AsRGBA32F();
-
 	// compute LinearAverage
-	FLinearColor LinearAverage;
-	{
-		FLinearColor LinearSum(0, 0, 0, 0);
-		for( int64 CurPixelIndex = 0; CurPixelIndex < NumPixels; ++CurPixelIndex )
-		{
-			LinearSum += ImageColors[ CurPixelIndex ];
-		}
-		LinearAverage = LinearSum / (float)NumPixels;
-	}
+	FLinearColor LinearAverage = FImageCore::ComputeImageLinearAverage(Image);
 
 	FLinearColor Scale(1, 1, 1, 1);
 
 	// we want to normalize the image to have 0.5 as average luminance, this is assuming clamping doesn't happen (can happen when using a very small Bokeh shape)
 	{
-		float RGBLum = (LinearAverage.R + LinearAverage.G + LinearAverage.B) / 3.0f;
+		float RGBLum = (LinearAverage.R + LinearAverage.G + LinearAverage.B) * (1.f/3.0f);
 
 		// ideally this would be 1 but then some pixels would need to be >1 which is not supported for the textureformat we want to use.
 		// The value affects the occlusion computation of the BokehDOF
@@ -2852,18 +2835,17 @@ static void ComputeBokehAlpha(FImage& Image)
 		Scale *= LumGoal / FMath::Max(RGBLum, 0.001f);
 	}
 
-	{
-		for( int64 CurPixelIndex = 0; CurPixelIndex < NumPixels; ++CurPixelIndex )
+	FImageCore::ImageParallelProcessLinearPixels(TEXT("PF.ComputeBokehAlpha"),Image,[&](TArrayView64<FLinearColor> & Colors) {
+		for( FLinearColor & Color : Colors )
 		{
-			const FLinearColor OriginalColor = ImageColors[ CurPixelIndex ];
-
 			// Convert to a linear color
-			FLinearColor LinearColor = OriginalColor * Scale;
-			float RGBLum = (LinearColor.R + LinearColor.G + LinearColor.B) / 3.0f;
+			FLinearColor LinearColor = Color * Scale;
+			float RGBLum = (LinearColor.R + LinearColor.G + LinearColor.B) * (1.f/3.0f);
 			LinearColor.A = FMath::Clamp(RGBLum, 0.0f, 1.0f);
-			ImageColors[ CurPixelIndex ] = LinearColor;
+			Color = LinearColor;
 		}
-	}
+		return FImageCore::ProcessLinearPixelsAction::Modified;
+	} );
 }
 
 /**
diff --git a/Engine/Source/Runtime/Engine/Private/TextureDerivedData.cpp b/Engine/Source/Runtime/Engine/Private/TextureDerivedData.cpp
--- a/Engine/Source/Runtime/Engine/Private/TextureDerivedData.cpp
+++ b/Engine/Source/Runtime/Engine/Private/TextureDerivedData.cpp
@@ -179,7 +179,7 @@ static void SerializeForKey(FArchive& Ar, const FTextureBuildSettings& Settings)
 		TempVector4f = Settings.AlphaCoverageThresholds; Ar << TempVector4f;
 	}
 	
-	TempByte = Settings.bComputeBokehAlpha; Ar << TempByte;
+	TempByte = Settings.bComputeBokehAlpha ? 2 : 0; Ar << TempByte;
 	TempByte = Settings.bReplicateRed; Ar << TempByte;
 	TempByte = Settings.bReplicateAlpha; Ar << TempByte;
 	TempByte = Settings.bDownsampleWithAverage; Ar << TempByte;
diff --git a/Engine/Source/Runtime/ImageCore/Private/ImageCore.cpp b/Engine/Source/Runtime/ImageCore/Private/ImageCore.cpp
--- a/Engine/Source/Runtime/ImageCore/Private/ImageCore.cpp
+++ b/Engine/Source/Runtime/ImageCore/Private/ImageCore.cpp
@@ -70,7 +70,7 @@ static inline int32 ParallelForComputeNumJobs(int64 & OutNumItemsPerJob,int64 Nu
 	// Always use ParallelFor with "Unbalanced"  so it does not do weird things to your job count
 
 	const int32 NumWorkers = FMath::Max(1, FTaskGraphInterface::Get().GetNumWorkerThreads());
-	int32 NumJobs = (int32)(NumItems / MinNumItemsPerJob); // round down
+	int32 NumJobs = IntCastChecked<int32>(NumItems / MinNumItemsPerJob); // round down
 
 	// NumJobs is now the maximum number of jobs we would want to do
 	//	(eg. one for each small packet of MinNumItemsPerJob)
@@ -118,6 +118,8 @@ static constexpr int64 PixelsPerJobAlignment = 64; // must be power of 2
 
 IMAGECORE_API int32 ImageParallelForComputeNumJobsForPixels(int64 & OutNumPixelsPerJob,int64 NumPixels)
 {
+	check( NumPixels > 0 );
+
 	int32 NumJobs = ParallelForComputeNumJobs(OutNumPixelsPerJob,NumPixels,MinPixelsPerJob,MinPixelsForAnyJob);
 	
 	if ( NumJobs > 1 )
@@ -128,16 +130,30 @@ IMAGECORE_API int32 ImageParallelForComputeNumJobsForPixels(int64 & OutNumPixels
 		// recompute NumJobs :
 		NumJobs = (NumPixels + OutNumPixelsPerJob-1) / OutNumPixelsPerJob;
 	}
+
+	check( NumJobs*OutNumPixelsPerJob >= NumPixels );
+	check( (NumJobs-1)*OutNumPixelsPerJob < NumPixels );
+
 	return NumJobs;
 }
 
-IMAGECORE_API int32 ImageParallelForComputeNumJobsForRows(int32 & OutNumItemsPerJob,int32 SizeX,int32 SizeY)
+IMAGECORE_API int32 ImageParallelForComputeNumJobsForRows(int32 & OutNumItemsPerJob,int64 SizeX,int64 SizeY)
 {
-	int64 NumPixels = int64(SizeX)*SizeY;
+	check( SizeX > 0 && SizeY > 0 );
+
+	int64 NumPixels = SizeX*SizeY;
 	int64 OutNumPixelsPerJob;
-	int32 NumJobs = ParallelForComputeNumJobs(OutNumPixelsPerJob,NumPixels,MinPixelsPerJob,MinPixelsForAnyJob);
-	OutNumItemsPerJob = (SizeY + NumJobs-1) / NumJobs; // round up;
-	return NumJobs;
+	int64 NumJobs1 = ParallelForComputeNumJobs(OutNumPixelsPerJob,NumPixels,MinPixelsPerJob,MinPixelsForAnyJob);
+	int64 OutNumItemsPerJob64 = (SizeY + NumJobs1-1) / NumJobs1;
+	
+	// recompute NumJobs :
+	int64 NumJobs = (SizeY + OutNumItemsPerJob64-1) / OutNumItemsPerJob64;
+
+	check( OutNumItemsPerJob64*NumJobs >= SizeY );
+	check( OutNumItemsPerJob64*(NumJobs-1) < SizeY );
+
+	OutNumItemsPerJob = IntCastChecked<int32>(OutNumItemsPerJob64);
+	return IntCastChecked<int32>(NumJobs);
 }
 
 template <typename Lambda>
diff --git a/Engine/Source/Runtime/ImageCore/Private/ImageParallelFor.cpp b/Engine/Source/Runtime/ImageCore/Private/ImageParallelFor.cpp
new file mode 100644
--- /dev/null
+++ b/Engine/Source/Runtime/ImageCore/Private/ImageParallelFor.cpp
@@ -0,0 +1,118 @@
+// Copyright Epic Games, Inc. All Rights Reserved.
+
+#include "ImageParallelFor.h"
+#include "HAL/CriticalSection.h"
+
+/**
+
+FImage is tight packed in memory with slices adjacent to each other
+so we can just treat it was a 2d image with Height *= NumSlices
+ 
+*/
+
+IMAGECORE_API int32 FImageCore::ImageParallelForComputeNumJobs(const FImageView & Image,int64 * pRowsPerJob)
+{
+	int64 SizeX = Image.SizeX;
+	int64 SizeY = (int64)Image.SizeY * Image.NumSlices;
+
+	int32 RowsPerJob;
+	int32 NumJobs = ImageParallelForComputeNumJobsForRows(RowsPerJob,SizeX,SizeY);
+
+	check( (int64)NumJobs * RowsPerJob >= SizeY );
+	check( (int64)(NumJobs-1) * RowsPerJob < SizeY );
+
+	*pRowsPerJob = RowsPerJob;
+	return NumJobs;
+}
+
+IMAGECORE_API void FImageCore::ImageParallelForMakePart(FImageView * Part,const FImageView & Image,int64 JobIndex,int64 RowsPerJob)
+{
+	int64 SizeX = Image.SizeX;
+	int64 SizeY = (int64)Image.SizeY * Image.NumSlices;
+
+	int64 StartY = JobIndex * RowsPerJob;
+	check( StartY < SizeY );
+
+	int64 EndY = FMath::Min(StartY + RowsPerJob,SizeY);
+
+	*Part = Image;
+	Part->SizeY = EndY - StartY;
+	Part->NumSlices = 1;
+	Part->RawData = (uint8 *)Image.RawData + Image.GetBytesPerPixel() * SizeX * StartY;
+}
+
+namespace
+{
+
+struct FLinearColorCmp
+{
+	inline bool operator () (const FLinearColor & lhs,const FLinearColor & rhs) const
+	{
+		if ( lhs.R != rhs.R ) return lhs.R < rhs.R;
+		if ( lhs.G != rhs.G ) return lhs.G < rhs.G;
+		if ( lhs.B != rhs.B ) return lhs.B < rhs.B;
+		if ( lhs.A != rhs.A ) return lhs.A < rhs.A;
+		return false;
+	}
+};
+
+static inline FLinearColor SumColors(const TArrayView64<FLinearColor> & Colors)
+{
+	VectorRegister4Float VecSum = VectorSetFloat1(0.f);
+	for ( FLinearColor & Color : Colors )
+	{
+		VecSum = VectorAdd(VecSum, VectorLoad(&Color.Component(0)));
+	}
+
+	FLinearColor Sum;
+	VectorStore(VecSum,&Sum.Component(0));
+	
+	return Sum;
+}
+
+};
+
+IMAGECORE_API FLinearColor FImageCore::ComputeImageLinearAverage(const FImageView & Image)
+{
+	FCriticalSection Accumulator_CritSec;
+	TArray<FLinearColor> Accumulator_Rows;
+	Accumulator_Rows.Reserve((int64)Image.SizeY*Image.NumSlices);
+
+	// just summing parallel portions to an accumulator would produce different output depending on thread count and timing
+	//	because the float sums to accumulator are not order and grouping invariant
+	// instead we are careful to ensure machine invariance
+	// the image is cut into rows
+	// each row is summed
+	// then all those row sums are sorted
+	// so they are accumulated in a fixed order
+
+	ImageParallelProcessLinearPixels(TEXT("PF.ComputeImageLinearAverage"),Image,
+		[&](TArrayView64<FLinearColor> & Colors) 
+		{
+			// this is called once per row
+			//	so it is always the same grouping of colors regardless of thread count
+			
+			FLinearColor Sum = SumColors(Colors);
+
+			// do not just += on accumulator here because that would be an order-dependent race that changes output
+			// instead we store all the row sums to later accumulate in known order
+
+			// this is a lock per row; a bit more contention than I would like
+			//	could multiplex it to reduce contention
+			Accumulator_CritSec.Lock();
+			Accumulator_Rows.Add(Sum);
+			Accumulator_CritSec.Unlock();
+
+			return ProcessLinearPixelsAction::ReadOnly;
+		}
+	);
+
+	Accumulator_Rows.Sort(FLinearColorCmp());
+
+	FLinearColor Accumulator = SumColors(Accumulator_Rows);
+
+	int64 NumPixels = Image.GetNumPixels();
+	Accumulator *= (1.f / NumPixels);
+
+	return Accumulator;
+}
\ No newline at end of file
diff --git a/Engine/Source/Runtime/ImageCore/Public/ImageCore.h b/Engine/Source/Runtime/ImageCore/Public/ImageCore.h
--- a/Engine/Source/Runtime/ImageCore/Public/ImageCore.h
+++ b/Engine/Source/Runtime/ImageCore/Public/ImageCore.h
@@ -380,6 +380,7 @@ public:
 	TArrayView64<FLinearColor> AsRGBA32F() const
 	{
 		check(Format == ERawImageFormat::RGBA32F);
+		check(GammaSpace == EGammaSpace::Linear);
 		return { (struct FLinearColor*)RawData, GetNumPixels() };
 	}
 
@@ -693,7 +694,7 @@ public:
 
 IMAGECORE_API int32 ImageParallelForComputeNumJobsForPixels(int64 & OutNumPixelsPerJob,int64 NumPixels);
 
-IMAGECORE_API int32 ImageParallelForComputeNumJobsForRows(int32 & OutNumItemsPerJob,int32 SizeX,int32 SizeY);
+IMAGECORE_API int32 ImageParallelForComputeNumJobsForRows(int32 & OutNumItemsPerJob,int64 SizeX,int64 SizeY);
 
 namespace FImageCore
 {
diff --git a/Engine/Source/Runtime/ImageCore/Public/ImageParallelFor.h b/Engine/Source/Runtime/ImageCore/Public/ImageParallelFor.h
new file mode 100644
--- /dev/null
+++ b/Engine/Source/Runtime/ImageCore/Public/ImageParallelFor.h
@@ -0,0 +1,131 @@
+// Copyright Epic Games, Inc. All Rights Reserved.
+
+/*=============================================================================
+ImageParallelFor.h: helper to run ParallelFor on images
+=============================================================================*/
+
+#pragma once
+
+#include "CoreMinimal.h"
+#include "ImageCore.h"
+#include "Async/ParallelFor.h"
+
+/**
+ * 
+
+ The easy and efficient way to write a pixel processing loop that supports all pixel formats is to use
+
+ ImageParallelProcessLinearPixels
+ 
+ passing in a lambda like :
+ ProcessLinearPixelsAction ProcessPixels(TArrayView64<FLinearColor> Colors)
+
+ this will let you visit all the pixels as FLinearColor.
+
+ You can use this for read (return ProcessLinearPixelsAction::ReadOnly) or read-write (return ProcessLinearPixelsAction::Modified).
+
+ */
+
+// ImageCore.h in global namespace:
+// ImageParallelForComputeNumJobsForPixels
+// ImageParallelForComputeNumJobsForRows
+
+namespace FImageCore
+{
+
+	IMAGECORE_API int32 ImageParallelForComputeNumJobs(const FImageView & Image,int64 * pRowsPerJob);
+	IMAGECORE_API void ImageParallelForMakePart(FImageView * Part,const FImageView & Whole,int64 JobIndex,int64 RowsPerJob);
+
+	// Lambda is a functor that works on an FImageView()
+	//	it will be called with portions of the image
+	//	each portion will be a 2d FImageView (NumSlices == 1)
+	// use like :
+	// FImageCore::ImageParallelFor( TEXT("Texture.AdjustImageColorsFunc.PF"),Image, [&](FImageView & ImagePart) {
+	template <typename Lambda>
+	inline void ImageParallelFor(const TCHAR* DebugName, const FImageView & Image, const Lambda& Func)
+	{
+		int64 RowsPerJob;
+		int32 NumJobs = ImageParallelForComputeNumJobs(Image,&RowsPerJob);
+		
+		ParallelFor(DebugName, NumJobs, 1, [=](int64 JobIndex)
+		{
+			FImageView Part;
+			ImageParallelForMakePart(&Part,Image,JobIndex,RowsPerJob);
+			Func(Part);
+		}, EParallelForFlags::Unbalanced);
+	}
+	
+	enum class ProcessLinearPixelsAction
+	{
+		ReadOnly,
+		Modified
+	};
+
+	/**
+	 * ProcessLinearPixels
+	 *  act on image pixels as TArrayView64<FLinearColor> Colors
+	 * pass a lambda that is :
+	 *  ProcessLinearPixelsAction ProcessPixels(TArrayView64<FLinearColor> Colors)
+	 *    return ProcessLinearPixelsAction::Modified if you modify colors
+	 *    return ProcessLinearPixelsAction::ReadOnly if you only read and don't change colors
+	 * your lambda will be called on one row at a time
+	 */
+	template <typename Lambda>
+	inline void ProcessLinearPixels(const FImageView & Image, const Lambda& Func)
+	{
+		if ( Image.Format == ERawImageFormat::RGBA32F )
+		{
+			// Image is already FLinearColor, can work on it in-place
+			TArrayView64<FLinearColor> Colors = Image.AsRGBA32F();
+			FLinearColor * Start = &Colors[0];
+
+			for(int64 Y=0;Y<(int64)Image.SizeY*Image.NumSlices;Y++)
+			{
+				TArrayView64<FLinearColor> RowColors(Start+Y*Image.SizeX,Image.SizeX);
+				Func(RowColors);
+			}
+		}
+		else
+		{
+			// convert to linear, act on linear, then copy back
+			FImage LinearRow(Image.SizeX,1,1,ERawImageFormat::RGBA32F,EGammaSpace::Linear);
+			TArrayView64<FLinearColor> RowColors = LinearRow.AsRGBA32F();
+			
+			for(int64 Y=0;Y<(int64)Image.SizeY*Image.NumSlices;Y++)
+			{
+				// point at one row of the image :
+				FImageView ImageRow = Image;
+				ImageRow.SizeY = 1;
+				ImageRow.NumSlices = 1;
+				ImageRow.RawData = (uint8 *)Image.RawData + Image.GetBytesPerPixel() * Y * Image.SizeX;
+
+				CopyImage(ImageRow,LinearRow);
+			
+				ProcessLinearPixelsAction Action = Func(RowColors);
+
+				if ( Action == ProcessLinearPixelsAction::Modified )
+				{
+					// Colors were modified, need to blit back
+					CopyImage(LinearRow,ImageRow);
+				}
+			}
+		}
+	}
+
+	// parallel process image as TArrayView64<FLinearColor>
+	//	your lambda is called on one row at a time so that the grouping is machine invariant
+	template <typename Lambda>
+	inline void ImageParallelProcessLinearPixels(const TCHAR* DebugName, const FImageView & Image, const Lambda& Func)
+	{
+		ImageParallelFor(DebugName,Image,[Func](const FImageView &Part) {
+			// call ProcessLinearPixels on each part :
+			ProcessLinearPixels(Part,Func);
+		} );
+	}
+
+	// ComputeImageLinearAverage
+	// compute the average linear color of the image
+	//	image can be any pixel format
+	//	parallel processing is used, but the result is not machine-dependent
+	IMAGECORE_API FLinearColor ComputeImageLinearAverage(const FImageView & Image);
+};
