ImageParallelFor ComputeImageLinearAverage remove unnecessary CritSec

#rnx

diff --git a/Engine/Source/Runtime/ImageCore/Private/ImageParallelFor.cpp b/Engine/Source/Runtime/ImageCore/Private/ImageParallelFor.cpp
--- a/Engine/Source/Runtime/ImageCore/Private/ImageParallelFor.cpp
+++ b/Engine/Source/Runtime/ImageCore/Private/ImageParallelFor.cpp
@@ -1,7 +1,6 @@
 // Copyright Epic Games, Inc. All Rights Reserved.
 
 #include "ImageParallelFor.h"
-#include "HAL/CriticalSection.h"
 
 /**
 
@@ -74,10 +73,11 @@ static inline FLinearColor SumColors(const TArrayView64<FLinearColor> & Colors)
 
 IMAGECORE_API FLinearColor FImageCore::ComputeImageLinearAverage(const FImageView & Image)
 {
-	FCriticalSection Accumulator_CritSec;
 	TArray<FLinearColor> Accumulator_Rows;
-	Accumulator_Rows.Reserve((int64)Image.SizeY*Image.NumSlices);
-
+	int64 Accumulator_RowCount = (int64)Image.SizeY*Image.NumSlices;
+	int64 Accumulator_RowIndex = 0;
+	Accumulator_Rows.SetNum(Accumulator_RowCount);
+	
 	// just summing parallel portions to an accumulator would produce different output depending on thread count and timing
 	//	because the float sums to accumulator are not order and grouping invariant
 	// instead we are careful to ensure machine invariance
@@ -97,16 +97,15 @@ IMAGECORE_API FLinearColor FImageCore::ComputeImageLinearAverage(const FImageVie
 			// do not just += on accumulator here because that would be an order-dependent race that changes output
 			// instead we store all the row sums to later accumulate in known order
 
-			// this is a lock per row; a bit more contention than I would like
-			//	could multiplex it to reduce contention
-			Accumulator_CritSec.Lock();
-			Accumulator_Rows.Add(Sum);
-			Accumulator_CritSec.Unlock();
+			int64 Index = FPlatformAtomics::InterlockedAdd(&Accumulator_RowIndex, 1);
+			Accumulator_Rows[Index] = Sum;
 
 			return ProcessLinearPixelsAction::ReadOnly;
 		}
 	);
 
+	check( Accumulator_RowIndex == Accumulator_RowCount );
+
 	Accumulator_Rows.Sort(FLinearColorCmp());
 
 	FLinearColor Accumulator = SumColors(Accumulator_Rows);
