ImageCore ResizeImage support resizing slices if slices count is the same

#rnx

diff --git a/Engine/Source/Runtime/ImageCore/Private/ImageCore.cpp b/Engine/Source/Runtime/ImageCore/Private/ImageCore.cpp
--- a/Engine/Source/Runtime/ImageCore/Private/ImageCore.cpp
+++ b/Engine/Source/Runtime/ImageCore/Private/ImageCore.cpp
@@ -1925,7 +1925,35 @@ IMAGECORE_API void FImageCore::ResizeImage(const FImageView & SourceImage,const
 		return;
 	}
 
-	check( SourceImage.NumSlices == 1 && DestImage.NumSlices == 1 ); // @todo Oodle : slices ?
+	if ( SourceImage.NumSlices != 1 || DestImage.NumSlices != 1 )
+	{
+		if ( SourceImage.NumSlices == DestImage.NumSlices )
+		{
+			// resize slices one by one (no filtering across slices)
+			//	eg. for arrays, not a 3d filter for cubes
+
+			ParallelFor(TEXT("ResizeImage.Slices"), SourceImage.NumSlices, 1, [&](int64 SliceIndex)
+			{
+				const FImageView SourceSlice = SourceImage.GetSlice(SliceIndex);
+				const FImageView DestSlice = DestImage.GetSlice(SliceIndex);
+				ResizeImage(SourceSlice,DestSlice,Filter);
+			}, EParallelForFlags::Unbalanced);
+		}
+		else
+		{
+			// different slice counts
+
+			UE_LOG(LogImageCore,Error,TEXT("ResizeImage: different slice count : %d -> %d"),
+				SourceImage.NumSlices,DestImage.NumSlices);
+
+			// zero out DestImage bytes?
+		}
+
+		return;
+	}
+
+	// single slice resize
+	check( SourceImage.NumSlices == 1 && DestImage.NumSlices == 1 );
 
 	if ( SourceImage.SizeX == DestImage.SizeX && SourceImage.SizeY == DestImage.SizeY &&
 		FilterIsNopWhenSameSize(Filter) )
@@ -1948,8 +1976,7 @@ IMAGECORE_API void FImageCore::ResizeImage(const FImageView & SourceImage,const
 	bool DestOk = GetFormatSTBIR(DestImage.Format,DestImage.GetGammaSpace(),
 		DestNumChannels,DestDataType,DestLayout);
 
-	if ( !SourceOk || !DestOk ||
-		( SourceNumChannels != DestNumChannels ) )
+	if ( !SourceOk || !DestOk || SourceNumChannels != DestNumChannels )
 	{
 		// Source and Dest formats have different channel count
 		//	it's hard to get this right in all cases through the stb_resize API
@@ -2045,7 +2072,7 @@ IMAGECORE_API void FImageCore::ResizeImageAllocDest(const FImageView & SourceIma
 	check( DestImage.RawData.Num() == 0 || SourceImage.RawData != DestImage.GetPixelPointer(0,0) ); // must not be resizing onto self
 
 	// note that if DestImage was already allocated to the right size, this will not re-allocate it
-	DestImage.Init(DestSizeX,DestSizeY,DestFormat,DestGammaSpace);
+	DestImage.Init(DestSizeX,DestSizeY,SourceImage.NumSlices,DestFormat,DestGammaSpace);
 	ResizeImage(SourceImage,DestImage,Filter);
 }
 
@@ -2056,7 +2083,8 @@ IMAGECORE_API void FImageCore::ResizeImageAllocDest(const FImageView & SourceIma
 
 IMAGECORE_API void FImageCore::ResizeImageInPlace(FImage & Image,int32 DestSizeX, int32 DestSizeY, ERawImageFormat::Type DestFormat, EGammaSpace DestGammaSpace, EResizeImageFilter Filter)
 {
-	if ( DestSizeX == Image.SizeX && DestSizeY == Image.SizeY && DestFormat == Image.Format && DestGammaSpace == Image.GetGammaSpace() )
+	if ( DestSizeX == Image.SizeX && DestSizeY == Image.SizeY && DestFormat == Image.Format && DestGammaSpace == Image.GetGammaSpace() &&
+		FilterIsNopWhenSameSize(Filter) )
 	{
 		// nop!
 		return;
