ImageCore ScaleChannelsSoMinMaxIsInZeroToOne replicates the behavior of the RenderTarget UNorm ReadPixels modifier
scales image outside [0,1] so that it is inside
uses ImageParallelFor

#rnx

diff --git a/Engine/Source/Developer/OpenColorIOWrapper/Private/OpenColorIOWrapper.cpp b/Engine/Source/Developer/OpenColorIOWrapper/Private/OpenColorIOWrapper.cpp
--- a/Engine/Source/Developer/OpenColorIOWrapper/Private/OpenColorIOWrapper.cpp
+++ b/Engine/Source/Developer/OpenColorIOWrapper/Private/OpenColorIOWrapper.cpp
@@ -932,7 +932,7 @@ bool FOpenColorIOWrapperProcessor::TransformImage(const FImageView& InOutImage)
 				// Apply parallelized color transformation (when it isn't a no-op)
 				if (!CPUProcessor->isNoOp())
 				{
-					FImageCore::ImageParallelFor(TEXT("FOpenColorIOWrapperProcessor.TransformImage.PF"), InOutImage, [&](FImageView& ImagePart)
+					FImageCore::ImageParallelFor(TEXT("FOpenColorIOWrapperProcessor.TransformImage.PF"), InOutImage, [&](FImageView& ImagePart,int64 RowY)
 						{
 							OCIO_NAMESPACE::PackedImageDesc ImagePartDesc = GetImageDesc(ImagePart);
 
diff --git a/Engine/Source/Developer/TextureCompressor/Private/TextureCompressorModule.cpp b/Engine/Source/Developer/TextureCompressor/Private/TextureCompressorModule.cpp
--- a/Engine/Source/Developer/TextureCompressor/Private/TextureCompressorModule.cpp
+++ b/Engine/Source/Developer/TextureCompressor/Private/TextureCompressorModule.cpp
@@ -2922,7 +2922,7 @@ void ITextureCompressorModule::AdjustImageColors(FImage& Image, const FTextureBu
 	{
 		TRACE_CPUPROFILER_EVENT_SCOPE(Texture.AdjustImageColors);
 
-		FImageCore::ImageParallelFor( TEXT("Texture.AdjustImageColorsFunc.PF"),Image, [&](FImageView & ImagePart)
+		FImageCore::ImageParallelFor( TEXT("Texture.AdjustImageColorsFunc.PF"),Image, [&](FImageView & ImagePart,int64 RowY)
 		{
 			TArrayView64<FLinearColor> ImageColors = ImagePart.AsRGBA32F();
 
@@ -2969,7 +2969,7 @@ static void ComputeBokehAlpha(FImage& Image)
 		Scale *= LumGoal / FMath::Max(RGBLum, 0.001f);
 	}
 
-	FImageCore::ImageParallelProcessLinearPixels(TEXT("PF.ComputeBokehAlpha"),Image,[&](TArrayView64<FLinearColor> & Colors) {
+	FImageCore::ImageParallelProcessLinearPixels(TEXT("PF.ComputeBokehAlpha"),Image,[&](TArrayView64<FLinearColor> & Colors,int64 RowY) {
 		for( FLinearColor & Color : Colors )
 		{
 			// Convert to a linear color
diff --git a/Engine/Source/Runtime/Engine/Private/TextureDerivedData.cpp b/Engine/Source/Runtime/Engine/Private/TextureDerivedData.cpp
--- a/Engine/Source/Runtime/Engine/Private/TextureDerivedData.cpp
+++ b/Engine/Source/Runtime/Engine/Private/TextureDerivedData.cpp
@@ -186,7 +186,8 @@ static void SerializeForKey(FArchive& Ar, const FTextureBuildSettings& Settings)
 		TempVector4f = Settings.AlphaCoverageThresholds; Ar << TempVector4f;
 	}
 	
-	TempByte = Settings.bComputeBokehAlpha ? 2 : 0; Ar << TempByte;
+	// Bokeh output version number bumped when processing changes
+	TempByte = Settings.bComputeBokehAlpha ? 3 : 0; Ar << TempByte;
 	TempByte = Settings.bReplicateRed; Ar << TempByte;
 	TempByte = Settings.bReplicateAlpha; Ar << TempByte;
 	TempByte = Settings.bDownsampleWithAverage; Ar << TempByte;
@@ -196,6 +197,7 @@ static void SerializeForKey(FArchive& Ar, const FTextureBuildSettings& Settings)
 
 		if(Settings.bSharpenWithoutColorShift && Settings.MipSharpening != 0.0f)
 		{
+			// @todo SerializeForKey these can go away whenever we bump the overall ddc key
 			// bSharpenWithoutColorShift prevented alpha sharpening. This got fixed
 			// Here we update the key to get those cases recooked.
 			TempByte = 2;
diff --git a/Engine/Source/Runtime/ImageCore/Private/ImageCore.cpp b/Engine/Source/Runtime/ImageCore/Private/ImageCore.cpp
--- a/Engine/Source/Runtime/ImageCore/Private/ImageCore.cpp
+++ b/Engine/Source/Runtime/ImageCore/Private/ImageCore.cpp
@@ -6,6 +6,7 @@
 #include "Serialization/CompactBinaryWriter.h"
 #include "TransferFunctions.h"
 #include "ColorSpace.h"
+#include "ImageParallelFor.h"
 
 DEFINE_LOG_CATEGORY_STATIC(LogImageCore, Log, All);
 
@@ -921,7 +922,7 @@ static void Bad_ResizeImage_Bilinear(const FImageView & SrcImage, const FImageVi
 	const float DestToSrcScaleX = (float)SrcImage.SizeX / (float)DestImage.SizeX;
 	const float DestToSrcScaleY = (float)SrcImage.SizeY / (float)DestImage.SizeY;
 
-	// @todo OodleImageResize : not a correct bilinear Resize?  missing 0.5 pixel center shift
+	// not a correct bilinear Resize?  missing 0.5 pixel center shift
 	//		deprecate this and redirect to new resizer
 	for (int64 DestY = 0; DestY < DestImage.SizeY; ++DestY)
 	{
@@ -1487,35 +1488,32 @@ IMAGECORE_API const FLinearColor ERawImageFormat::GetOnePixelLinear(const void *
 }
 
 
-void FImageCore::SanitizeFloat16AndSetAlphaOpaqueForBC6H(const FImageView & InOutImage)
+void FImageCore::SanitizeFloat16AndSetAlphaOpaqueForBC6H(const FImageView & InOutImageWhole)
 {
 	TRACE_CPUPROFILER_EVENT_SCOPE(Texture.SanitizeFloat16AndSetAlphaOpaqueForBC6H);
-	check(InOutImage.Format == ERawImageFormat::RGBA16F);
+	check(InOutImageWhole.Format == ERawImageFormat::RGBA16F);
 
-	const int64 TexelNum = InOutImage.GetNumPixels();
-	FFloat16Color* Data = reinterpret_cast<FFloat16Color*>(InOutImage.RawData);
-
-	// @todo Oodle : parallel for ?
-	for (int64 TexelIndex = 0; TexelIndex < TexelNum; ++TexelIndex)
+	ImageParallelFor(TEXT("PF.SanitizeFloat16AndSetAlphaOpaqueForBC6H"),InOutImageWhole,[](const FImageView & InOutImagePart,int64 RowY)
 	{
-		FFloat16Color& F16Color = Data[TexelIndex];
+		const int64 TexelNum = InOutImagePart.GetNumPixels();
+		FFloat16Color* Data = reinterpret_cast<FFloat16Color*>(InOutImagePart.RawData);
 
-		F16Color.R = F16Color.R.GetClampedNonNegativeAndFinite();
-		F16Color.G = F16Color.G.GetClampedNonNegativeAndFinite();
-		F16Color.B = F16Color.B.GetClampedNonNegativeAndFinite();
-		F16Color.A.SetOne();
-	}
+		for (int64 TexelIndex = 0; TexelIndex < TexelNum; ++TexelIndex)
+		{
+			FFloat16Color& F16Color = Data[TexelIndex];
+
+			F16Color.R = F16Color.R.GetClampedNonNegativeAndFinite();
+			F16Color.G = F16Color.G.GetClampedNonNegativeAndFinite();
+			F16Color.B = F16Color.B.GetClampedNonNegativeAndFinite();
+			F16Color.A.SetOne();
+		}
+	} );
 }
 
 bool FImageCore::DetectAlphaChannel(const FImageView & InImage)
 {
 	TRACE_CPUPROFILER_EVENT_SCOPE(Texture.DetectAlphaChannel);
 
-	// previous :
-	// "SMALL_NUMBER" is quite small, this provides almost no tolerance
-	// #define SMALL_NUMBER		(1.e-8f)
-	//const float FloatNonOpaqueAlpha = 1.0f - SMALL_NUMBER;
-	// 
 	// opaque alpha threshold where we'd quantize to < 255 in U8
 	//  images with only alpha larger than this are treated as opaque
 	const float FloatNonOpaqueAlpha = 254.5f / 255.f; // the U8 alpha threshold
@@ -1593,10 +1591,8 @@ bool FImageCore::DetectAlphaChannel(const FImageView & InImage)
 	return false;
 }
 
-void FImageCore::SetAlphaOpaque(const FImageView & InImage)
+static void SetAlphaOpaque_SingleThreaded(const FImageView & InImage)
 {
-	// this is not currently parallel but easily could be
-
 	int64 NumPixels = (int64)InImage.SizeX * InImage.SizeY * InImage.NumSlices;
 
 	if (InImage.Format == ERawImageFormat::BGRA8)
@@ -1668,37 +1664,11 @@ void FImageCore::SetAlphaOpaque(const FImageView & InImage)
 	}
 }
 
-
-void FImageCore::ComputeChannelLinearMinMax(const FImageView & InImage, FLinearColor & OutMin, FLinearColor & OutMax)
+void FImageCore::SetAlphaOpaque(const FImageView & InImage)
 {
-	// @todo Oodle : for speed, we should ideally scan the image for min/max in its native pixel format
-	//	then only convert the min/max colors to float linear after the scan
-	//	don't convert the whole image
-
-	FImage ImageLinear;
-	InImage.CopyTo(ImageLinear,ERawImageFormat::RGBA32F,EGammaSpace::Linear);
-	
-	TArrayView64<FLinearColor> Colors = ImageLinear.AsRGBA32F();
-	if ( Colors.Num() == 0 )
-	{
-		OutMin = FLinearColor(ForceInit);
-		OutMax = FLinearColor(ForceInit);
-		return;
-	}
-	
-	VectorRegister4Float VMin = VectorLoad(&Colors[0].Component(0));
-	VectorRegister4Float VMax = VMin;
-	
-	for ( const FLinearColor & Color : Colors )
-	{
-		VectorRegister4Float VCur = VectorLoad(&Color.Component(0));
-
-		VMin = VectorMin(VMin,VCur);
-		VMax = VectorMax(VMax,VCur);
-	}
-
-	VectorStore(VMin,&OutMin.Component(0));
-	VectorStore(VMax,&OutMax.Component(0));
+	ImageParallelFor(TEXT("PF.SetAlphaOpaque"),InImage,[](const FImageView & Part,int64 Row) {
+		SetAlphaOpaque_SingleThreaded(Part);
+	} );
 }
 
 void FImageCore::TransformToWorkingColorSpace(const FImageView& InLinearImage, const FVector2d& SourceRedChromaticity, const FVector2d& SourceGreenChromaticity, const FVector2d& SourceBlueChromaticity, const FVector2d& SourceWhiteChromaticity, UE::Color::EChromaticAdaptationMethod Method, double EqualityTolerance)
diff --git a/Engine/Source/Runtime/ImageCore/Private/ImageParallelFor.cpp b/Engine/Source/Runtime/ImageCore/Private/ImageParallelFor.cpp
--- a/Engine/Source/Runtime/ImageCore/Private/ImageParallelFor.cpp
+++ b/Engine/Source/Runtime/ImageCore/Private/ImageParallelFor.cpp
@@ -6,13 +6,19 @@
 
 FImage is tight packed in memory with slices adjacent to each other
 so we can just treat it was a 2d image with Height *= NumSlices
- 
+
+@todo Oodle : actually because of the tight packed property, there's no need to use the original image dimensions at all
+we could just cut into 1d runs of the desired pixel count.
+That would give better parallelism on skew images (than keeping original SizeX which we do now).
+
+eg. make "ImagePart" of 16384 pixels, and make "Rows" for the FLinearColor pass that are always exactly 512 pixels
+
 */
 
 IMAGECORE_API int32 FImageCore::ImageParallelForComputeNumJobs(const FImageView & Image,int64 * pRowsPerJob)
 {
 	int64 SizeX = Image.SizeX;
-	int64 SizeY = (int64)Image.SizeY * Image.NumSlices;
+	int64 SizeY = ImageParallelForComputeNumRows(Image);
 
 	int32 RowsPerJob;
 	int32 NumJobs = ImageParallelForComputeNumJobsForRows(RowsPerJob,SizeX,SizeY);
@@ -24,10 +30,10 @@ IMAGECORE_API int32 FImageCore::ImageParallelForComputeNumJobs(const FImageView
 	return NumJobs;
 }
 
-IMAGECORE_API void FImageCore::ImageParallelForMakePart(FImageView * Part,const FImageView & Image,int64 JobIndex,int64 RowsPerJob)
+IMAGECORE_API int64 FImageCore::ImageParallelForMakePart(FImageView * Part,const FImageView & Image,int64 JobIndex,int64 RowsPerJob)
 {
 	int64 SizeX = Image.SizeX;
-	int64 SizeY = (int64)Image.SizeY * Image.NumSlices;
+	int64 SizeY = ImageParallelForComputeNumRows(Image);
 
 	int64 StartY = JobIndex * RowsPerJob;
 	check( StartY < SizeY );
@@ -38,6 +44,8 @@ IMAGECORE_API void FImageCore::ImageParallelForMakePart(FImageView * Part,const
 	Part->SizeY = EndY - StartY;
 	Part->NumSlices = 1;
 	Part->RawData = (uint8 *)Image.RawData + Image.GetBytesPerPixel() * SizeX * StartY;
+
+	return StartY;
 }
 
 namespace
@@ -74,8 +82,7 @@ static inline FLinearColor SumColors(const TArrayView64<FLinearColor> & Colors)
 IMAGECORE_API FLinearColor FImageCore::ComputeImageLinearAverage(const FImageView & Image)
 {
 	TArray<FLinearColor> Accumulator_Rows;
-	int64 Accumulator_RowCount = (int64)Image.SizeY*Image.NumSlices;
-	int64 Accumulator_RowIndex = 0;
+	int64 Accumulator_RowCount = ImageParallelForComputeNumRows(Image);
 	Accumulator_Rows.SetNum(Accumulator_RowCount);
 	
 	// just summing parallel portions to an accumulator would produce different output depending on thread count and timing
@@ -83,11 +90,10 @@ IMAGECORE_API FLinearColor FImageCore::ComputeImageLinearAverage(const FImageVie
 	// instead we are careful to ensure machine invariance
 	// the image is cut into rows
 	// each row is summed
-	// then all those row sums are sorted
-	// so they are accumulated in a fixed order
+	// then all those row sums are accumulated
 
 	ImageParallelProcessLinearPixels(TEXT("PF.ComputeImageLinearAverage"),Image,
-		[&](TArrayView64<FLinearColor> & Colors) 
+		[&](TArrayView64<FLinearColor> & Colors,int64 Y) 
 		{
 			// this is called once per row
 			//	so it is always the same grouping of colors regardless of thread count
@@ -97,21 +103,147 @@ IMAGECORE_API FLinearColor FImageCore::ComputeImageLinearAverage(const FImageVie
 			// do not just += on accumulator here because that would be an order-dependent race that changes output
 			// instead we store all the row sums to later accumulate in known order
 
-			int64 Index = FPlatformAtomics::InterlockedAdd(&Accumulator_RowIndex, 1);
-			Accumulator_Rows[Index] = Sum;
+			Accumulator_Rows[Y] = Sum;
 
 			return ProcessLinearPixelsAction::ReadOnly;
 		}
 	);
 
-	check( Accumulator_RowIndex == Accumulator_RowCount );
-
-	Accumulator_Rows.Sort(FLinearColorCmp());
-
 	FLinearColor Accumulator = SumColors(Accumulator_Rows);
 
 	int64 NumPixels = Image.GetNumPixels();
 	Accumulator *= (1.f / NumPixels);
 
 	return Accumulator;
-}
\ No newline at end of file
+}
+
+namespace
+{
+
+	struct FMinMax
+	{
+		VectorRegister4Float VMin,VMax;
+
+		FMinMax() { }
+
+		FMinMax(VectorRegister4Float InMin,VectorRegister4Float InMax) : VMin(InMin), VMax(InMax)
+		{
+		}
+	};
+	
+	static inline FMinMax MinMax(const FMinMax & A,const FMinMax & B)
+	{
+		return FMinMax(
+			VectorMin(A.VMin,B.VMin),
+			VectorMax(A.VMax,B.VMax) );
+	}
+
+	static inline FMinMax MinMaxColors(const TArrayView64<FLinearColor> & Colors)
+	{
+		check( Colors.Num() > 0 );
+
+		VectorRegister4Float VMin = VectorLoad(&Colors[0].Component(0));
+		VectorRegister4Float VMax = VMin;
+	
+		for ( const FLinearColor & Color : Colors )
+		{
+			VectorRegister4Float VCur = VectorLoad(&Color.Component(0));
+
+			VMin = VectorMin(VMin,VCur);
+			VMax = VectorMax(VMax,VCur);
+		}
+
+		return FMinMax(VMin,VMax);
+	}
+
+};
+
+void FImageCore::ComputeChannelLinearMinMax(const FImageView & Image, FLinearColor & OutMin, FLinearColor & OutMax)
+{
+	if ( Image.GetNumPixels() == 0 )
+	{
+		OutMin = FLinearColor(ForceInit);
+		OutMax = FLinearColor(ForceInit);
+		return;
+	}
+	
+	TArray<FMinMax> MinMax_Rows;
+	int64 MinMax_RowCount = ImageParallelForComputeNumRows(Image);
+	MinMax_Rows.SetNum(MinMax_RowCount);
+	
+	ImageParallelProcessLinearPixels(TEXT("PF.ComputeChannelLinearMinMax"),Image,
+		[&](TArrayView64<FLinearColor> & Colors,int64 Y) 
+		{
+			FMinMax Sum = MinMaxColors(Colors);
+
+			MinMax_Rows[Y] = Sum;
+
+			return ProcessLinearPixelsAction::ReadOnly;
+		}
+	);
+	
+	// now MinMax on all the rows :
+	FMinMax NetMinMax = MinMax_Rows[0];
+	for (const FMinMax & MM_Row : MinMax_Rows )
+	{
+		NetMinMax = MinMax(NetMinMax, MM_Row );
+	}
+
+	VectorStore(NetMinMax.VMin,&OutMin.Component(0));
+	VectorStore(NetMinMax.VMax,&OutMax.Component(0));
+}
+
+bool FImageCore::ScaleChannelsSoMinMaxIsInZeroToOne(const FImageView & Image)
+{
+	if ( Image.GetNumPixels() == 0 )
+	{
+		return false;
+	}
+	if ( ! ERawImageFormat::IsHDR(Image.Format) )
+	{
+		// early out : if Image is U8/U16 it is already in [0,1]
+		return false;
+	}
+
+	FLinearColor Min,Max;
+	ComputeChannelLinearMinMax(Image,Min,Max);
+
+	if ( Min.R >= 0.f && Min.G >= 0.f && Min.B >= 0.f && Min.A >= 0.f &&
+		Max.R <= 1.f && Max.G <= 1.f && Max.B <= 1.f && Max.A <= 1.f )
+	{
+		// nothing to do
+		return false;
+	}
+
+	VectorRegister4Float VMin = VectorLoad(&Min.Component(0));
+	VectorRegister4Float VMax = VectorLoad(&Max.Component(0));
+
+	// this makes it so that the end of the range that was already in [0,1] is not modified :
+	VMin = VectorMin( VMin, MakeVectorRegisterFloat(0.f,0.f,0.f,0.f) );
+	VMax = VectorMax( VMax, MakeVectorRegisterFloat(1.f,1.f,1.f,1.f) );
+
+	// VScale = 1.f/(Max-Min)
+	VectorRegister4Float VSub = VectorSubtract(VMax,VMin);
+	// avoid divide by zero :
+	VSub = VectorMax(VSub, MakeVectorRegisterFloat(FLT_MIN,FLT_MIN,FLT_MIN,FLT_MIN) );
+	VectorRegister4Float VScale = VectorReciprocalAccurate(VSub);
+
+	ImageParallelProcessLinearPixels(TEXT("PF.ScaleChannelsSoMinMaxIsInZeroToOne"),Image,
+		[&](TArrayView64<FLinearColor> & Colors,int64 Y) 
+		{
+			for ( FLinearColor & Color : Colors )
+			{
+				VectorRegister4Float VCur = VectorLoad(&Color.Component(0));
+
+				VCur = VectorSubtract(VCur,VMin);
+				VCur = VectorMultiply(VCur,VScale);
+
+				VectorStore(VCur, &Color.Component(0));
+			}
+
+			return ProcessLinearPixelsAction::Modified;
+		}
+	);
+
+	return true;
+}
diff --git a/Engine/Source/Runtime/ImageCore/Public/ImageCore.h b/Engine/Source/Runtime/ImageCore/Public/ImageCore.h
--- a/Engine/Source/Runtime/ImageCore/Public/ImageCore.h
+++ b/Engine/Source/Runtime/ImageCore/Public/ImageCore.h
@@ -811,6 +811,29 @@ IMAGECORE_API void ResizeTo(const FImageView & SourceImage,FImage& DestImage, in
  */
 IMAGECORE_API void ComputeChannelLinearMinMax(const FImageView & InImage, FLinearColor & OutMin, FLinearColor & OutMax);
 
+/**
+* If the image has any values outside the [0,1] range, rescale that edge of the domain so that it is in [0,1]
+* does not affect images that were previously in [0,1]
+*
+* also does not change the side of the domain that is not out of bounds
+* eg. values in [0.25,200.0] will be rescaled to [0.25,1.0]
+*
+* returns bool if any change was made
+*
+* If the input format is U8 or U16, no change will ever be made and this will return false.
+*
+* This can be useful if you want to save an HDR/float image to a U8 image format for visualization.
+* This is equivalent to what's called the "UNorm" transformation by the RenderTarget ReadPixels functions.
+*/
+IMAGECORE_API bool ScaleChannelsSoMinMaxIsInZeroToOne(const FImageView & ImageToModify);
+
+/** ComputeImageLinearAverage
+ * compute the average linear color of the image
+ *	image can be any pixel format
+ *	parallel processing is used, but the result is not machine-dependent
+ */
+IMAGECORE_API FLinearColor ComputeImageLinearAverage(const FImageView & Image);
+
 /**
  * Apply a color space transformation from the source chromaticities to the engine's working color space.
  *
diff --git a/Engine/Source/Runtime/ImageCore/Public/ImageParallelFor.h b/Engine/Source/Runtime/ImageCore/Public/ImageParallelFor.h
--- a/Engine/Source/Runtime/ImageCore/Public/ImageParallelFor.h
+++ b/Engine/Source/Runtime/ImageCore/Public/ImageParallelFor.h
@@ -32,9 +32,14 @@ ImageParallelFor.h: helper to run ParallelFor on images
 
 namespace FImageCore
 {
+	// ImageParallelFor treats the image as being 2d (no slices) but with *NumSlices rows
+	inline int64 ImageParallelForComputeNumRows(const FImageView & Image)
+	{
+		return (int64)Image.SizeY * Image.NumSlices;
+	}
 
 	IMAGECORE_API int32 ImageParallelForComputeNumJobs(const FImageView & Image,int64 * pRowsPerJob);
-	IMAGECORE_API void ImageParallelForMakePart(FImageView * Part,const FImageView & Whole,int64 JobIndex,int64 RowsPerJob);
+	IMAGECORE_API int64 ImageParallelForMakePart(FImageView * Part,const FImageView & Whole,int64 JobIndex,int64 RowsPerJob);
 
 	// Lambda is a functor that works on an FImageView()
 	//	it will be called with portions of the image
@@ -50,8 +55,8 @@ namespace FImageCore
 		ParallelFor(DebugName, NumJobs, 1, [=](int64 JobIndex)
 		{
 			FImageView Part;
-			ImageParallelForMakePart(&Part,Image,JobIndex,RowsPerJob);
-			Func(Part);
+			int64 StartY = ImageParallelForMakePart(&Part,Image,JobIndex,RowsPerJob);
+			Func(Part,StartY);
 		}, EParallelForFlags::Unbalanced);
 	}
 	
@@ -71,7 +76,7 @@ namespace FImageCore
 	 * your lambda will be called on one row at a time
 	 */
 	template <typename Lambda>
-	inline void ProcessLinearPixels(const FImageView & Image, const Lambda& Func)
+	inline void ProcessLinearPixels(const FImageView & Image, const Lambda& Func, int64 StartY)
 	{
 		if ( Image.Format == ERawImageFormat::RGBA32F )
 		{
@@ -79,10 +84,10 @@ namespace FImageCore
 			TArrayView64<FLinearColor> Colors = Image.AsRGBA32F();
 			FLinearColor * Start = &Colors[0];
 
-			for(int64 Y=0;Y<(int64)Image.SizeY*Image.NumSlices;Y++)
+			for(int64 Y=0;Y<ImageParallelForComputeNumRows(Image);Y++)
 			{
 				TArrayView64<FLinearColor> RowColors(Start+Y*Image.SizeX,Image.SizeX);
-				Func(RowColors);
+				Func(RowColors,StartY+Y);
 			}
 		}
 		else
@@ -91,7 +96,7 @@ namespace FImageCore
 			FImage LinearRow(Image.SizeX,1,1,ERawImageFormat::RGBA32F,EGammaSpace::Linear);
 			TArrayView64<FLinearColor> RowColors = LinearRow.AsRGBA32F();
 			
-			for(int64 Y=0;Y<(int64)Image.SizeY*Image.NumSlices;Y++)
+			for(int64 Y=0;Y<ImageParallelForComputeNumRows(Image);Y++)
 			{
 				// point at one row of the image :
 				FImageView ImageRow = Image;
@@ -101,7 +106,7 @@ namespace FImageCore
 
 				CopyImage(ImageRow,LinearRow);
 			
-				ProcessLinearPixelsAction Action = Func(RowColors);
+				ProcessLinearPixelsAction Action = Func(RowColors,StartY+Y);
 
 				if ( Action == ProcessLinearPixelsAction::Modified )
 				{
@@ -117,15 +122,10 @@ namespace FImageCore
 	template <typename Lambda>
 	inline void ImageParallelProcessLinearPixels(const TCHAR* DebugName, const FImageView & Image, const Lambda& Func)
 	{
-		ImageParallelFor(DebugName,Image,[Func](const FImageView &Part) {
+		ImageParallelFor(DebugName,Image,[Func](const FImageView &Part,int64 StartY) {
 			// call ProcessLinearPixels on each part :
-			ProcessLinearPixels(Part,Func);
+			ProcessLinearPixels(Part,Func,StartY);
 		} );
 	}
 
-	// ComputeImageLinearAverage
-	// compute the average linear color of the image
-	//	image can be any pixel format
-	//	parallel processing is used, but the result is not machine-dependent
-	IMAGECORE_API FLinearColor ComputeImageLinearAverage(const FImageView & Image);
 };
