ImageCore CopyImage clean up excessive insights trace spew

diff --git a/Engine/Source/Runtime/ImageCore/Private/ImageCore.cpp b/Engine/Source/Runtime/ImageCore/Private/ImageCore.cpp
--- a/Engine/Source/Runtime/ImageCore/Private/ImageCore.cpp
+++ b/Engine/Source/Runtime/ImageCore/Private/ImageCore.cpp
@@ -162,15 +162,29 @@ IMAGECORE_API int32 ImageParallelForComputeNumJobsForRows(int32 & OutNumItemsPer
 template <typename Lambda>
 static void ParallelLoop(const TCHAR* DebugName, int32 NumJobs, int64 TexelsPerJob, int64 NumTexels, const Lambda& Func)
 {
-	ParallelFor(DebugName, NumJobs, 1, [=](int64 JobIndex)
+	if ( NumJobs <= 1 )
 	{
-		const int64 StartIndex = JobIndex * TexelsPerJob;
-		const int64 EndIndex = FMath::Min(StartIndex + TexelsPerJob, NumTexels);
-		for (int64 TexelIndex = StartIndex; TexelIndex < EndIndex; ++TexelIndex)
+		// special case for non-parallel
+		//	this is mainly to avoid making insights traces for this case
+		//	would be nice if ParallelFor just had a very early out for this
+
+		for (int64 TexelIndex = 0; TexelIndex < NumTexels; ++TexelIndex)
 		{
 			Func(TexelIndex);
 		}
-	}, EParallelForFlags::Unbalanced);
+	}
+	else
+	{
+		ParallelFor(DebugName, NumJobs, 1, [=](int64 JobIndex)
+		{
+			const int64 StartIndex = JobIndex * TexelsPerJob;
+			const int64 EndIndex = FMath::Min(StartIndex + TexelsPerJob, NumTexels);
+			for (int64 TexelIndex = StartIndex; TexelIndex < EndIndex; ++TexelIndex)
+			{
+				Func(TexelIndex);
+			}
+		}, EParallelForFlags::Unbalanced);
+	}
 }
 
 // ParallelOr : call Func() on all texels ; returns true if Func is true for any texel
@@ -381,7 +395,8 @@ IMAGECORE_API void FImageCore::CopyImage(const FImageView & SrcImage,const FImag
 		return;
 	}
 
-	TRACE_CPUPROFILER_EVENT_SCOPE(Texture.CopyImage);
+	bool bDoTrace = SrcImage.GetNumPixels() > 8192;
+	TRACE_CPUPROFILER_EVENT_SCOPE_TEXT_CONDITIONAL("Texture.CopyImage",bDoTrace);
 
 	check(SrcImage.IsImageInfoValid());
 	check(DestImage.IsImageInfoValid());
@@ -391,7 +406,7 @@ IMAGECORE_API void FImageCore::CopyImage(const FImageView & SrcImage,const FImag
 	if ( SrcImage.Format == DestImage.Format &&
 		SrcImage.GammaSpace == DestImage.GammaSpace )
 	{
-		TRACE_CPUPROFILER_EVENT_SCOPE(Texture.CopyImage.memcpy);
+		//TRACE_CPUPROFILER_EVENT_SCOPE(Texture.CopyImage.memcpy);
 
 		int64 Bytes = SrcImage.GetImageSizeBytes();
 		check( DestImage.GetImageSizeBytes() == Bytes );
