ImageCore::CopyImage TempLinear path, use per-line temp row, not whole image temp surface

diff --git a/Engine/Source/Runtime/ImageCore/Private/ImageCore.cpp b/Engine/Source/Runtime/ImageCore/Private/ImageCore.cpp
--- a/Engine/Source/Runtime/ImageCore/Private/ImageCore.cpp
+++ b/Engine/Source/Runtime/ImageCore/Private/ImageCore.cpp
@@ -396,7 +396,7 @@ IMAGECORE_API void FImageCore::CopyImage(const FImageView & SrcImage,const FImag
 		return;
 	}
 
-	bool bDoTrace = SrcImage.GetNumPixels() > 8192;
+	bool bDoTrace = SrcImage.GetNumPixels() > 16384;
 	TRACE_CPUPROFILER_EVENT_SCOPE_TEXT_CONDITIONAL("Texture.CopyImage",bDoTrace);
 
 	check(SrcImage.IsImageInfoValid());
@@ -879,9 +879,34 @@ IMAGECORE_API void FImageCore::CopyImage(const FImageView & SrcImage,const FImag
 		// Arbitrary conversion, use 32-bit linear float as an intermediate format.
 		// this is unnecessarily expensive to do something like G8 to R16F, but rare
 		// if this shows up as a hot spot, identify the formats using this path and add direct conversions between them
-		FImage TempImage(SrcImage.SizeX, SrcImage.SizeY, SrcImage.NumSlices, ERawImageFormat::RGBA32F, EGammaSpace::Linear);
-		FImageCore::CopyImage(SrcImage, TempImage);
-		FImageCore::CopyImage(TempImage, DestImage);
+
+		if ( SrcImage.GetNumPixels()*16 < 65536 ) // RGBA32F fits in 64K
+		{
+			// whole image temp
+
+			FImage TempImage(SrcImage.SizeX, SrcImage.SizeY, SrcImage.NumSlices, ERawImageFormat::RGBA32F, EGammaSpace::Linear);
+			FImageCore::CopyImage(SrcImage, TempImage);
+			FImageCore::CopyImage(TempImage, DestImage);
+		}
+		else
+		{
+			// use per-line temp, not whole image temp!
+
+			ImageParallelFor(TEXT("PF.CopyImage.TempLinear"),SrcImage,[&](const FImageView & SrcImagePart,int64 StartY)
+			{
+				FImage TempRow(SrcImage.SizeX, 1, 1, ERawImageFormat::RGBA32F, EGammaSpace::Linear);
+				int64 NumRows = SrcImagePart.SizeY;
+
+				for(int64 RowY = StartY; RowY < StartY+NumRows; RowY++)
+				{
+					FImageView SrcRow  = ImageParallelForGetOneRowView(SrcImage,RowY);
+					FImageView DestRow = ImageParallelForGetOneRowView(DestImage,RowY);
+				
+					FImageCore::CopyImage(SrcRow, TempRow);
+					FImageCore::CopyImage(TempRow,DestRow);
+				}
+			});
+		}
 	}
 }
 
diff --git a/Engine/Source/Runtime/ImageCore/Public/ImageParallelFor.h b/Engine/Source/Runtime/ImageCore/Public/ImageParallelFor.h
--- a/Engine/Source/Runtime/ImageCore/Public/ImageParallelFor.h
+++ b/Engine/Source/Runtime/ImageCore/Public/ImageParallelFor.h
@@ -37,6 +37,19 @@ namespace FImageCore
 	{
 		return (int64)Image.SizeY * Image.NumSlices;
 	}
+	
+	// Y indexes across slices and SizeY here :
+	inline FImageView ImageParallelForGetOneRowView(const FImageView & Image,int64 Y)
+	{
+		check( Y >= 0 && Y < ImageParallelForComputeNumRows(Image) );
+
+		// point at one row of the image :
+		FImageView ImageRow = Image;
+		ImageRow.SizeY = 1;
+		ImageRow.NumSlices = 1;
+		ImageRow.RawData = (uint8 *)Image.RawData + Image.GetBytesPerPixel() * Y * Image.SizeX;
+		return ImageRow;
+	}
 
 	IMAGECORE_API int32 ImageParallelForComputeNumJobs(const FImageView & Image,int64 * pRowsPerJob);
 	IMAGECORE_API int64 ImageParallelForMakePart(FImageView * Part,const FImageView & Whole,int64 JobIndex,int64 RowsPerJob);
@@ -99,10 +112,7 @@ namespace FImageCore
 			for(int64 Y=0;Y<ImageParallelForComputeNumRows(Image);Y++)
 			{
 				// point at one row of the image :
-				FImageView ImageRow = Image;
-				ImageRow.SizeY = 1;
-				ImageRow.NumSlices = 1;
-				ImageRow.RawData = (uint8 *)Image.RawData + Image.GetBytesPerPixel() * Y * Image.SizeX;
+				FImageView ImageRow = ImageParallelForGetOneRowView(Image,Y);
 
 				CopyImage(ImageRow,LinearRow);
 			
