ImageParallelFor: avoid TArray iterator

diff --git a/Engine/Source/Runtime/ImageCore/Private/ImageParallelFor.cpp b/Engine/Source/Runtime/ImageCore/Private/ImageParallelFor.cpp
--- a/Engine/Source/Runtime/ImageCore/Private/ImageParallelFor.cpp
+++ b/Engine/Source/Runtime/ImageCore/Private/ImageParallelFor.cpp
@@ -184,7 +184,7 @@ void FImageCore::ComputeChannelLinearMinMax(const FImageView & Image, FLinearCol
 	
 	// now MinMax on all the rows :
 	FMinMax NetMinMax = MinMax_Rows[0];
-	for (const FMinMax & MM_Row : MinMax_Rows )
+	for (const FMinMax & MM_Row : TArrayView64<FMinMax>(MinMax_Rows) ) // TArray iterator is slow, must wrap in ArrayView
 	{
 		NetMinMax = MinMax(NetMinMax, MM_Row );
 	}
