ImageCore: more resize filter options (Lanczos)
make DownscaleOptions Default pass through without changing to SimpleAverage for new filters
map DownscaleOptions to better resize filters

diff --git a/Engine/Source/Developer/TextureCompressor/Private/TextureCompressorModule.cpp b/Engine/Source/Developer/TextureCompressor/Private/TextureCompressorModule.cpp
--- a/Engine/Source/Developer/TextureCompressor/Private/TextureCompressorModule.cpp
+++ b/Engine/Source/Developer/TextureCompressor/Private/TextureCompressorModule.cpp
@@ -1507,17 +1507,62 @@ static void DownscaleImage(const FImage& SrcImage, FImage& DstImage, const FText
 				
 		// choose Filter from ETextureDownscaleOptions
 		FImageCore::EResizeImageFilter Filter;
-
-		if ( Settings.DownscaleOptions == (uint8)ETextureDownscaleOptions::Unfiltered )
+		
+		// unlike MinGen, the ETextureDownscaleOptions does not have "Blur" options, only Sharpen
+		/*0=no sharpening but better quality which is softer, 1=little, 5=medium, 10=extreme. */
+		
+		switch( (ETextureDownscaleOptions)Settings.DownscaleOptions )
+		{
+		case ETextureDownscaleOptions::Default:
+			Filter = FImageCore::EResizeImageFilter::CubicMitchell;
+			break;
+		case ETextureDownscaleOptions::Unfiltered:
 			Filter = FImageCore::EResizeImageFilter::PointSample;
-		else if ( Settings.DownscaleOptions == (uint8)ETextureDownscaleOptions::SimpleAverage ) // <- Default maps to SimpleAverage
+			break;
+		case ETextureDownscaleOptions::SimpleAverage:
 			Filter = FImageCore::EResizeImageFilter::Triangle;
-		else if ( Settings.DownscaleOptions >= (uint8)ETextureDownscaleOptions::Sharpen5 ) // sharpest
+			break;
+			
+		// cubic Mitchells in order of B: (from highest B to lowest; eg. smoothest to sharpest)
+		// CubicGaussian  // B = 1
+		// CubicMitchell // B = 1/3
+		// MitchellOneQuarter // B = 1/4
+		// MitchellOneSixth // B = 1/6
+		// CubicSharp  // B = 0
+
+		case ETextureDownscaleOptions::Sharpen0:
+			Filter = FImageCore::EResizeImageFilter::CubicGaussian;
+			break;
+		case ETextureDownscaleOptions::Sharpen1:
+		case ETextureDownscaleOptions::Sharpen2:
+			Filter = FImageCore::EResizeImageFilter::CubicMitchell;
+			break;
+		case ETextureDownscaleOptions::Sharpen3:
+		case ETextureDownscaleOptions::Sharpen4:
+			Filter = FImageCore::EResizeImageFilter::MitchellOneQuarter;
+			break;
+		case ETextureDownscaleOptions::Sharpen5:
+		case ETextureDownscaleOptions::Sharpen6:
+			Filter = FImageCore::EResizeImageFilter::MitchellOneSixth;
+			break;			
+		case ETextureDownscaleOptions::Sharpen7:
+		case ETextureDownscaleOptions::Sharpen8:
 			Filter = FImageCore::EResizeImageFilter::CubicSharp;
-		else if ( Settings.DownscaleOptions >= (uint8)ETextureDownscaleOptions::Sharpen0 ) // sharp
-			Filter = FImageCore::EResizeImageFilter::AdaptiveSharp;		 // winds up as CubicMitchell
-		else
-			Filter = FImageCore::EResizeImageFilter::AdaptiveSmooth; // <- this can't be reached, Blur is not available as a DownscaleOptions
+			break;
+		case ETextureDownscaleOptions::Sharpen9:
+			//Filter = FImageCore::EResizeImageFilter::Lanczos4;
+			Filter = FImageCore::EResizeImageFilter::MitchellNegOneSixth;
+			break;
+		case ETextureDownscaleOptions::Sharpen10:
+			//Filter = FImageCore::EResizeImageFilter::Lanczos5;
+			Filter = FImageCore::EResizeImageFilter::MitchellNegOneThird;
+			break;
+
+		default:
+			UE_LOG(LogTextureCompressor, Warning, TEXT("Downscale option not recognized : %d"),(int)Settings.DownscaleOptions);
+			Filter = FImageCore::EResizeImageFilter::CubicMitchell;
+			break;
+		}
 
 		FImage Temp; // needed because Src == Dst
 		FImageCore::ResizeImageAllocDest(SrcImage,Temp,FinalSizeX,FinalSizeY,Filter);
@@ -1552,7 +1597,7 @@ static void DownscaleImage(const FImage& SrcImage, FImage& DstImage, const FText
 	AllocateTempForMips(TempData, SrcImage.SizeX, SrcImage.SizeY, FMath::Max(1, SrcImage.SizeX / 2), FMath::Max(1, SrcImage.SizeY / 2), false, AvgKernel, 2, false, bUnfiltered, bUseNewMipFilter);
 
 	int32 NumIterations = 0;
-	while(Downscale > 2.0f)
+	while(Downscale > 2.0f) // NOT == 2.0
 	{
 		int32 DstSizeX = FMath::Max(1, ImageChain[0]->SizeX / 2 );
 		int32 DstSizeY = FMath::Max(1, ImageChain[0]->SizeY / 2 );
@@ -1592,14 +1637,21 @@ static void DownscaleImage(const FImage& SrcImage, FImage& DstImage, const FText
 	
 	int32 KernelSize = 2;
 	float Sharpening = 0.0f;
+	bool bBilinear = false;
 	if (Settings.DownscaleOptions >= (uint8)ETextureDownscaleOptions::Sharpen0 && Settings.DownscaleOptions <= (uint8)ETextureDownscaleOptions::Sharpen10)
 	{
 		// 0 .. 2.0f
 		Sharpening = (float)((int32)Settings.DownscaleOptions - (int32)ETextureDownscaleOptions::Sharpen0) * 0.2f;
 		KernelSize = 8;
 	}
-	
-	bool bBilinear = Settings.DownscaleOptions == (uint8)ETextureDownscaleOptions::SimpleAverage;
+	else
+	{
+		// Default should have mapped to SimpleAverage before getting here (in GetDownscaleOptions)
+		check( Settings.DownscaleOptions != (uint8)ETextureDownscaleOptions::Default );
+
+		bBilinear = Settings.DownscaleOptions == (uint8)ETextureDownscaleOptions::SimpleAverage;
+		check( bBilinear || bUnfiltered );
+	}	
 	
 	FImageKernel2D KernelSharpen;
 	KernelSharpen.BuildSeparatableGaussWithSharpen(KernelSize, Sharpening);
@@ -1641,6 +1693,12 @@ static void DownscaleImage(const FImage& SrcImage, FImage& DstImage, const FText
 			}
 			else if(bBilinear)
 			{
+				// in the common case of Downscale == 2.0
+				//	this is actually not bilinar at all, or a 2x2 downsample
+				//	it's just point sampled
+				//	because of the way the pixel centers are not offset correctly
+				//	this is just sampling pixels at 0,2,4,..
+				//FilteredColor = LookupSourceMipBilinear(SrcImageData, SourceX + 0.5, SourceY + 0.5); // <- correct offsets for Downscale == 2.0
 				FilteredColor = LookupSourceMipBilinear(SrcImageData, SourceX, SourceY);
 			}
 			else
diff --git a/Engine/Source/Runtime/Engine/Classes/Engine/TextureLODSettings.h b/Engine/Source/Runtime/Engine/Classes/Engine/TextureLODSettings.h
--- a/Engine/Source/Runtime/Engine/Classes/Engine/TextureLODSettings.h
+++ b/Engine/Source/Runtime/Engine/Classes/Engine/TextureLODSettings.h
@@ -37,7 +37,7 @@ struct FTextureLODGroup
 		, HighPriorityLoad(false)
 		, DuplicateNonOptionalMips(false)
 		, Downscale(1.0)
-		, DownscaleOptions(ETextureDownscaleOptions::SimpleAverage)
+		, DownscaleOptions(ETextureDownscaleOptions::Default)
 		, VirtualTextureTileCountBias(0)
 		, VirtualTextureTileSizeBias(0)
 		, LossyCompressionAmount(TLCA_Default)
diff --git a/Engine/Source/Runtime/Engine/Private/TextureDerivedData.cpp b/Engine/Source/Runtime/Engine/Private/TextureDerivedData.cpp
--- a/Engine/Source/Runtime/Engine/Private/TextureDerivedData.cpp
+++ b/Engine/Source/Runtime/Engine/Private/TextureDerivedData.cpp
@@ -264,6 +264,8 @@ static void SerializeForKey(FArchive& Ar, const FTextureBuildSettings& Settings)
 			// downscale behavior changed to use ResizeImage
 			TempName = TEXTURE_DDC_STB_IMAGE_RESIZE_VERSION;
 			Ar << TempName;
+			TempName = TEXT("Downscale ResizeImage changed 02-29-2024");
+			Ar << TempName;
 		}
 	}
 
diff --git a/Engine/Source/Runtime/Engine/Private/TextureLODSettings.cpp b/Engine/Source/Runtime/Engine/Private/TextureLODSettings.cpp
--- a/Engine/Source/Runtime/Engine/Private/TextureLODSettings.cpp
+++ b/Engine/Source/Runtime/Engine/Private/TextureLODSettings.cpp
@@ -338,7 +338,15 @@ void UTextureLODSettings::GetDownscaleOptions(const UTexture& Texture, const ITa
 	ETextureDownscaleOptions GroupDownscaleOptions = TextureLODGroups[Texture.LODGroup].DownscaleOptions;
 	if (GroupDownscaleOptions == ETextureDownscaleOptions::Default)
 	{
-		GroupDownscaleOptions = ETextureDownscaleOptions::SimpleAverage;
+		if ( Texture.bUseNewMipFilter )
+		{
+			// leave Default alone
+		}
+		else
+		{
+			// legacy behavior, change Default to SimpleAverage
+			GroupDownscaleOptions = ETextureDownscaleOptions::SimpleAverage;
+		}
 	}
 		
 	Downscale = Texture.Downscale.GetValueForPlatform(*CurrentPlatform.IniPlatformName());
diff --git a/Engine/Source/Runtime/ImageCore/Private/ImageCore.cpp b/Engine/Source/Runtime/ImageCore/Private/ImageCore.cpp
--- a/Engine/Source/Runtime/ImageCore/Private/ImageCore.cpp
+++ b/Engine/Source/Runtime/ImageCore/Private/ImageCore.cpp
@@ -1918,6 +1918,12 @@ static bool FilterIsNopWhenSameSize(EResizeImageFilter FilterWithFlags)
 		case EResizeImageFilter::CubicGaussian:
 		case EResizeImageFilter::CubicSharp:
 		case EResizeImageFilter::CubicMitchell:
+		case EResizeImageFilter::MitchellOneQuarter:
+		case EResizeImageFilter::MitchellOneSixth:
+		case EResizeImageFilter::MitchellNegOneSixth:
+		case EResizeImageFilter::MitchellNegOneThird:
+		case EResizeImageFilter::Lanczos4:
+		case EResizeImageFilter::Lanczos5:
 			return false;
 
 		default: // all enum values should be explicitly listed above
@@ -1926,7 +1932,88 @@ static bool FilterIsNopWhenSameSize(EResizeImageFilter FilterWithFlags)
 	}
 }
 
-static stbir_filter MapFilterToStb(EResizeImageFilter FilterWithFlags,int64 SizeFm,int64 SizeTo,uint32 WrapFlag, stbir_edge & OutStbirEdgeMode)
+// B in [0,5/3]
+//	smaller B = sharper, larger B = blurrier
+//  B = 0 is the catmull-rom interpolating b-spline
+//	B = 1/3 is standard compromise Mitchell
+//  B = 1 is the cubic Gaussian
+// in between pulse_cubic_mitchell_B = lerp( catrom, cubic, B )
+static float pulse_cubic_mitchell_B(const float signed_x,const float B)
+{
+	float ax = signed_x < 0.0 ? -signed_x : signed_x;
+	if (ax<1.0) return (  ax*ax* ( (9.0-6.0*B)*ax - 15.0 + 9.0*B ) + 6.0 - 2.0*B ) *(1.0/6.0);
+	else if (ax<2.0) return ( ax*ax* ((2.0*B - 3.0)*ax + (15.0 - 9.0*B)) + (12.0*B - 24.0)*ax + (12.0 - 4.0*B) ) *(1.0/6.0);
+	else return 0.0;
+}
+
+// sharp (upsample) :
+template <int denom> 
+static float my_filter_mitchell_B(float x, float s, void * user_data)
+{
+	const float B = 1.f/denom;
+
+	return (float) pulse_cubic_mitchell_B(x,B);
+}
+
+static float my_pulse_sinc(const float x) 
+{
+	const float pix = (3.141592653590f) * x;
+
+	const float tiny = 1e-4f;
+	if ( fabsf(pix) < tiny )
+	{
+		// Taylor series of sin for x small
+		return 1.f - (pix*pix)*(1/6.f);
+	}
+	else
+	{
+		return sinf(pix) / pix;
+	}
+}
+
+static float my_filter_lanczos(float x, float support)
+{
+	if ( fabsf(x) >= support ) return 0.f;
+
+	float pulse = my_pulse_sinc(x);
+	float window = my_pulse_sinc(x * (1.f/support));
+	return pulse * window;
+}
+
+static float my_filter_lanczos4(float x, float s, void * user_data)
+{
+	return my_filter_lanczos(x,2.f);
+}
+static float my_filter_lanczos5(float x, float s, void * user_data)
+{	
+	return my_filter_lanczos(x,2.5f);
+}
+
+static float my_filter_support_two(float s, void * user_data) 
+{
+	return 2;
+}
+static float my_filter_support_twoandhalf(float s, void * user_data) 
+{
+	return 2.5f;
+}
+
+struct stbir_kernel_and_support
+{
+	stbir__kernel_callback * kernel = nullptr;
+	stbir__support_callback * support = nullptr;
+};
+
+static const stbir_kernel_and_support c_my_ks_mitchell_onequarter = { my_filter_mitchell_B<4> , my_filter_support_two };
+static const stbir_kernel_and_support c_my_ks_mitchell_onesixth = { my_filter_mitchell_B<6> , my_filter_support_two };
+static const stbir_kernel_and_support c_my_ks_mitchell_negonesixth = { my_filter_mitchell_B<-6> , my_filter_support_two };
+static const stbir_kernel_and_support c_my_ks_mitchell_negonethird = { my_filter_mitchell_B<-3> , my_filter_support_two };
+static const stbir_kernel_and_support c_my_ks_lanczos4 = { my_filter_lanczos4 , my_filter_support_two };
+static const stbir_kernel_and_support c_my_ks_lanczos5 = { my_filter_lanczos5 , my_filter_support_twoandhalf };
+
+
+static stbir_filter MapFilterToStb(EResizeImageFilter FilterWithFlags,int64 SizeFm,int64 SizeTo,uint32 WrapFlag, stbir_edge & OutStbirEdgeMode,
+	stbir_kernel_and_support & OutCustomKS)
 {
 	if ( SizeFm == SizeTo && FilterIsNopWhenSameSize(FilterWithFlags) )
 	{
@@ -1937,6 +2024,9 @@ static stbir_filter MapFilterToStb(EResizeImageFilter FilterWithFlags,int64 Size
 	OutStbirEdgeMode = WrapFlag ? STBIR_EDGE_WRAP : STBIR_EDGE_CLAMP;
 	
 	EResizeImageFilter Filter = FilterWithFlags & EResizeImageFilter::WithoutFlagsMask;
+	
+	//const stbir_filter STBIR_FILTER_OutCustomKS = STBIR_FILTER_OTHER; // <- not allowed
+	const stbir_filter STBIR_FILTER_OutCustomKS = STBIR_FILTER_DEFAULT;
 
 	switch(Filter)
 	{
@@ -1946,6 +2036,25 @@ static stbir_filter MapFilterToStb(EResizeImageFilter FilterWithFlags,int64 Size
 		case EResizeImageFilter::CubicSharp: return STBIR_FILTER_CATMULLROM;
 		case EResizeImageFilter::CubicMitchell: return STBIR_FILTER_MITCHELL;
 		case EResizeImageFilter::PointSample: return STBIR_FILTER_POINT_SAMPLE;
+		
+		case EResizeImageFilter::MitchellOneQuarter:
+			OutCustomKS = c_my_ks_mitchell_onequarter;
+			return STBIR_FILTER_OutCustomKS;
+		case EResizeImageFilter::MitchellOneSixth:
+			OutCustomKS = c_my_ks_mitchell_onesixth;
+			return STBIR_FILTER_OutCustomKS;
+		case EResizeImageFilter::MitchellNegOneSixth:
+			OutCustomKS = c_my_ks_mitchell_negonesixth;
+			return STBIR_FILTER_OutCustomKS;
+		case EResizeImageFilter::MitchellNegOneThird:
+			OutCustomKS = c_my_ks_mitchell_negonethird;
+			return STBIR_FILTER_OutCustomKS;
+		case EResizeImageFilter::Lanczos4:
+			OutCustomKS = c_my_ks_lanczos4;
+			return STBIR_FILTER_OutCustomKS;
+		case EResizeImageFilter::Lanczos5:
+			OutCustomKS = c_my_ks_lanczos5;
+			return STBIR_FILTER_OutCustomKS;
 
 		case EResizeImageFilter::Default: // AdaptiveSharp matches STBIR_FILTER_DEFAULT
 		case EResizeImageFilter::AdaptiveSharp:
@@ -2070,15 +2179,18 @@ IMAGECORE_API void FImageCore::ResizeImage(const FImageView & SourceImage,const
 		SourceImage.RawData,SourceImage.SizeX,SourceImage.SizeY,SourceImage.GetBytesPerPixel()*SourceImage.SizeX,
 		DestImage.RawData,DestImage.SizeX,DestImage.SizeY,DestImage.GetBytesPerPixel()*DestImage.SizeX,
 		SourceLayout,SourceDataType);
+		
+	stbir_set_pixel_layouts(&resize,SourceLayout,DestLayout);
+	stbir_set_datatypes(&resize,SourceDataType,DestDataType);
 
 	stbir_edge StbirEdgeX,StbirEdgeY;
-	stbir_filter StbirFilterX = MapFilterToStb(Filter,SourceImage.SizeX,DestImage.SizeX, (uint32)(Filter & EResizeImageFilter::Flag_WrapX), StbirEdgeX);
-	stbir_filter StbirFilterY = MapFilterToStb(Filter,SourceImage.SizeY,DestImage.SizeY, (uint32)(Filter & EResizeImageFilter::Flag_WrapY), StbirEdgeY);
+	stbir_kernel_and_support StbirKSX,StbirKSY;
+	stbir_filter StbirFilterX = MapFilterToStb(Filter,SourceImage.SizeX,DestImage.SizeX, (uint32)(Filter & EResizeImageFilter::Flag_WrapX), StbirEdgeX, StbirKSX);
+	stbir_filter StbirFilterY = MapFilterToStb(Filter,SourceImage.SizeY,DestImage.SizeY, (uint32)(Filter & EResizeImageFilter::Flag_WrapY), StbirEdgeY, StbirKSY);
 
+	stbir_set_edgemodes(&resize, StbirEdgeX, StbirEdgeY);	
 	stbir_set_filters(&resize, StbirFilterX, StbirFilterY);
-	stbir_set_pixel_layouts(&resize,SourceLayout,DestLayout);
-	stbir_set_datatypes(&resize,SourceDataType,DestDataType);
-	stbir_set_edgemodes(&resize, StbirEdgeX, StbirEdgeY);
+	stbir_set_filter_callbacks(&resize,StbirKSX.kernel,StbirKSX.support,StbirKSY.kernel,StbirKSY.support);
 
 	const int32 NumWorkers = FTaskGraphInterface::Get().GetNumWorkerThreads();
 
diff --git a/Engine/Source/Runtime/ImageCore/Public/ImageCore.h b/Engine/Source/Runtime/ImageCore/Public/ImageCore.h
--- a/Engine/Source/Runtime/ImageCore/Public/ImageCore.h
+++ b/Engine/Source/Runtime/ImageCore/Public/ImageCore.h
@@ -863,11 +863,27 @@ IMAGECORE_API void TransformToWorkingColorSpace(const FImageView& InLinearImage,
 		Triangle, 
 		Bilinear = Triangle, // synonym
 		CubicGaussian, // smooth Mitchell B=1,C=0, B-spline, Gaussian-like
-		CubicSharp, // sharp interpolating cubic, Catmull-ROM (has negative lobes)
+		CubicSharp, // sharp interpolating cubic, Catmull-Rom (has negative lobes) (Mitchell B=0)
 		CubicMitchell, // compromise between sharp and smooth cubic, Mitchell-Netrevalli filter with B=1/3, C=1/3 (has negative lobes)
 		AdaptiveSharp,  // sharper adaptive filter; uses CubicSharp for upsample and CubicMitchell for downsample, nop for same size
 		AdaptiveSmooth,  // smoother adaptive filter; uses CubicMitchell for upsample and CubicGaussian for downsample, nop for same size
 
+		MitchellOneQuarter, // B=1/4 ("Robidoux")
+		MitchellOneSixth, // B=1/6
+		MitchellNegOneSixth, // B=-1/6 , over-sharpening
+		MitchellNegOneThird, // B=-1/3 , over-sharpening
+		Lanczos4, // sharp
+		Lanczos5, // very sharp, often ringy
+
+		// cubic Mitchells in order of B: (from highest B to lowest; eg. smoothest to sharpest)
+		// CubicGaussian  // B = 1
+		// CubicMitchell // B = 1/3
+		// MitchellOneQuarter // B = 1/4
+		// MitchellOneSixth // B = 1/6
+		// CubicSharp  // B = 0
+		// MitchellNegOneSixth // B = -1/6
+		// MitchellNegOneThird // B = -1/3
+
 		WithoutFlagsMask = 63,
 		Flag_WrapX = 64,  // default edge mode is clamp; set these to wrap instead
 		Flag_WrapY = 128
