ImageCore: ResizeImage add option to do alpha-weighted resize filter

diff --git a/Engine/Source/Runtime/ImageCore/Private/ImageCore.cpp b/Engine/Source/Runtime/ImageCore/Private/ImageCore.cpp
--- a/Engine/Source/Runtime/ImageCore/Private/ImageCore.cpp
+++ b/Engine/Source/Runtime/ImageCore/Private/ImageCore.cpp
@@ -746,10 +746,10 @@ IMAGECORE_API void FImageCore::CopyImage(const FImageView & SrcImage,const FImag
 				{
 					int64 SrcIndex = TexelIndex * 4;
 					DestColors[TexelIndex] = FLinearColor(
-						SrcColors[SrcIndex + 0] / 65535.0f,
-						SrcColors[SrcIndex + 1] / 65535.0f,
-						SrcColors[SrcIndex + 2] / 65535.0f,
-						SrcColors[SrcIndex + 3] / 65535.0f
+						SrcColors[SrcIndex + 0] * (1.f/65535.0f),
+						SrcColors[SrcIndex + 1] * (1.f/65535.0f),
+						SrcColors[SrcIndex + 2] * (1.f/65535.0f),
+						SrcColors[SrcIndex + 3] * (1.f/65535.0f)
 					);
 				});
 			}
@@ -1071,7 +1071,6 @@ void FImage::ChangeFormat(ERawImageFormat::Type DestFormat, EGammaSpace DestGamm
 	{
 		FImage Temp;
 		CopyTo(Temp,DestFormat,DestGammaSpace);
-		//*this = MoveTemp(Temp); // or swap?
 		Swap(Temp);
 	}
 }
@@ -1327,10 +1326,10 @@ void FImage::Linearize(uint8 SourceEncoding, FImage& DestImage) const
 			{
 				int64 SrcIndex = TexelIndex * 4;
 				FLinearColor Color(
-					SrcColors[SrcIndex + 0] / 65535.0f,
-					SrcColors[SrcIndex + 1] / 65535.0f,
-					SrcColors[SrcIndex + 2] / 65535.0f,
-					SrcColors[SrcIndex + 3] / 65535.0f
+					SrcColors[SrcIndex + 0] * (1.f/65535.0f),
+					SrcColors[SrcIndex + 1] * (1.f/65535.0f),
+					SrcColors[SrcIndex + 2] * (1.f/65535.0f),
+					SrcColors[SrcIndex + 3] * (1.f/65535.0f)
 				);
 				Color = DecodeFunction(Color);
 				DestColors[TexelIndex] = SaturateToHalfFloat(Color);
@@ -1486,6 +1485,11 @@ IMAGECORE_API bool ERawImageFormat::IsHDR(Type Format)
 	return Format == RGBA16F || Format == RGBA32F || Format == R16F || Format == R32F || Format == BGRE8;
 }
 
+IMAGECORE_API bool ERawImageFormat::HasAlphaChannel(Type Format)
+{
+	return Format == BGRA8 || Format == RGBA16 || Format == RGBA16F || Format == RGBA32F;
+}
+
 IMAGECORE_API const FLinearColor ERawImageFormat::GetOnePixelLinear(const void * PixelData,ERawImageFormat::Type Format,EGammaSpace Gamma)
 {
 	switch(Format)
@@ -1582,6 +1586,11 @@ void FImageCore::SanitizeFloat16AndSetAlphaOpaqueForBC6H(const FImageView & InOu
 
 bool FImageCore::DetectAlphaChannel(const FImageView & InImage)
 {
+	if ( ! HasAlphaChannel(InImage.Format) )
+	{
+		return false;
+	}
+
 	TRACE_CPUPROFILER_EVENT_SCOPE(Texture.DetectAlphaChannel);
 
 	// opaque alpha threshold where we'd quantize to < 255 in U8
@@ -1644,17 +1653,10 @@ bool FImageCore::DetectAlphaChannel(const FImageView & InImage)
 			return (ColorPtr[TexelIndex].A.GetFloat() <= FloatNonOpaqueAlpha);
 		});
 	}
-	else if (InImage.Format == ERawImageFormat::G8 ||
-		InImage.Format == ERawImageFormat::BGRE8 ||
-		InImage.Format == ERawImageFormat::G16 ||
-		InImage.Format == ERawImageFormat::R16F ||
-		InImage.Format == ERawImageFormat::R32F)
-	{
-		// source image formats don't have alpha
-	}
 	else
 	{
 		// new format ?
+		//	formats without HasAlphaChannel were previously excluded
 		check(0);
 	}
 
@@ -1736,6 +1738,11 @@ static void SetAlphaOpaque_SingleThreaded(const FImageView & InImage)
 
 void FImageCore::SetAlphaOpaque(const FImageView & InImage)
 {
+	if ( ! HasAlphaChannel(InImage.Format) )
+	{
+		return;
+	}
+
 	ImageParallelFor(TEXT("PF.SetAlphaOpaque"),InImage,[](const FImageView & Part,int64 Row) {
 		SetAlphaOpaque_SingleThreaded(Part);
 	} );
@@ -1901,6 +1908,26 @@ static bool GetFormatSTBIR(ERawImageFormat::Type Format,EGammaSpace GammaSpace,
 	}
 }
 
+static stbir_pixel_layout ChangeToAlphaWeighted(stbir_pixel_layout layout)
+{
+	// incoming layout is not alpha weighted
+
+	// we initially made 4-channel RGBA formats with the _PM flag
+	// to tell STBIR (falsely) that they are pre-multiplied already
+	// that makes STBIR do non-alpha-weighted resize, which is our default
+	// instead now take off the _PM flag
+	// that tells STBIR to do alpha-weighted resize
+
+	switch(layout)
+	{
+	case STBIR_RGBA_PM: return STBIR_RGBA;
+	case STBIR_BGRA_PM: return STBIR_BGRA;
+	default:
+		check(0); // we're only called on 4-channel layouts
+		return layout;
+	}
+}
+
 static bool FilterIsNopWhenSameSize(EResizeImageFilter FilterWithFlags)
 {
 	EResizeImageFilter Filter = FilterWithFlags & EResizeImageFilter::WithoutFlagsMask;
@@ -1940,19 +1967,29 @@ static bool FilterIsNopWhenSameSize(EResizeImageFilter FilterWithFlags)
 // in between pulse_cubic_mitchell_B = lerp( catrom, cubic, B )
 static float pulse_cubic_mitchell_B(const float signed_x,const float B)
 {
-	float ax = signed_x < 0.0 ? -signed_x : signed_x;
-	if (ax<1.0) return (  ax*ax* ( (9.0-6.0*B)*ax - 15.0 + 9.0*B ) + 6.0 - 2.0*B ) *(1.0/6.0);
-	else if (ax<2.0) return ( ax*ax* ((2.0*B - 3.0)*ax + (15.0 - 9.0*B)) + (12.0*B - 24.0)*ax + (12.0 - 4.0*B) ) *(1.0/6.0);
-	else return 0.0;
+	const float ax = fabsf(signed_x);
+	if (ax<1.f)
+	{
+		return ax*ax* ( ((9/6.f)-(6/6.f)*B)*ax - (15/6.f) + (9/6.f)*B ) + ((6/6.f) - (2/6.f)*B);
+	}
+	else if (ax<2.f)
+	{
+		return ax*ax* ( ((2/6.f)*B - (3/6.f))*ax + ((15/6.f) - (9/6.f)*B) ) + ((12/6.f)*B - (24/6.f))*ax + ((12/6.f) - (4/6.f)*B);
+	}
+	else
+	{
+		return 0.0f;
+	}
 }
 
-// sharp (upsample) :
+// template on integer denominator to make Mitchell filters of various B :
+//	(note denom can be negative but not zero)
 template <int denom> 
 static float my_filter_mitchell_B(float x, float s, void * user_data)
 {
 	const float B = 1.f/denom;
 
-	return (float) pulse_cubic_mitchell_B(x,B);
+	return pulse_cubic_mitchell_B(x,B);
 }
 
 static float my_pulse_sinc(const float x) 
@@ -1971,9 +2008,10 @@ static float my_pulse_sinc(const float x)
 	}
 }
 
+// "support" must match the support function used
 static float my_filter_lanczos(float x, float support)
 {
-	if ( fabsf(x) >= support ) return 0.f;
+	if ( fabsf(x) >= support ) return 0.f; // ? is this necessary, or does STBIR do this already ?
 
 	float pulse = my_pulse_sinc(x);
 	float window = my_pulse_sinc(x * (1.f/support));
@@ -2179,7 +2217,16 @@ IMAGECORE_API void FImageCore::ResizeImage(const FImageView & SourceImage,const
 		SourceImage.RawData,SourceImage.SizeX,SourceImage.SizeY,SourceImage.GetBytesPerPixel()*SourceImage.SizeX,
 		DestImage.RawData,DestImage.SizeX,DestImage.SizeY,DestImage.GetBytesPerPixel()*DestImage.SizeX,
 		SourceLayout,SourceDataType);
-		
+	
+	// we only use STBIR with the same channel count in src and dest
+	//	so questions about how PM/RGBA translate to 1CHANNEL are moot
+	check( SourceNumChannels == DestNumChannels );
+	if ( (uint32)( Filter & EResizeImageFilter::Flag_AlphaWeighted ) && SourceNumChannels == 4 )
+	{
+		SourceLayout = ChangeToAlphaWeighted(SourceLayout);
+		DestLayout = ChangeToAlphaWeighted(DestLayout);
+	}
+	
 	stbir_set_pixel_layouts(&resize,SourceLayout,DestLayout);
 	stbir_set_datatypes(&resize,SourceDataType,DestDataType);
 
diff --git a/Engine/Source/Runtime/ImageCore/Public/ImageCore.h b/Engine/Source/Runtime/ImageCore/Public/ImageCore.h
--- a/Engine/Source/Runtime/ImageCore/Public/ImageCore.h
+++ b/Engine/Source/Runtime/ImageCore/Public/ImageCore.h
@@ -78,6 +78,7 @@ namespace ERawImageFormat
 	IMAGECORE_API bool GetFormatFromString(FUtf8StringView InString, Type& OutFormat);
 	
 	IMAGECORE_API bool IsHDR(Type Format);
+	IMAGECORE_API bool HasAlphaChannel(Type Format);
 	
 	// Get one pixel of Format type from PixelData and return in Linear color
 	IMAGECORE_API const FLinearColor GetOnePixelLinear(const void * PixelData,Type Format,EGammaSpace Gamma);
@@ -886,7 +887,8 @@ IMAGECORE_API void TransformToWorkingColorSpace(const FImageView& InLinearImage,
 
 		WithoutFlagsMask = 63,
 		Flag_WrapX = 64,  // default edge mode is clamp; set these to wrap instead
-		Flag_WrapY = 128
+		Flag_WrapY = 128,
+		Flag_AlphaWeighted = 256 // weight RGB by A in the filter; that is, convert non-premultiplied RGBA input to premultiplied for filtering, then convert back
 	};
 	ENUM_CLASS_FLAGS(EResizeImageFilter);
 
