ImageCore: fix check in ResizeImage when source is Pow22 and dest gamma was not specified (dest cannot be pow22)

diff --git a/Engine/Source/Runtime/ImageCore/Private/ImageCore.cpp b/Engine/Source/Runtime/ImageCore/Private/ImageCore.cpp
--- a/Engine/Source/Runtime/ImageCore/Private/ImageCore.cpp
+++ b/Engine/Source/Runtime/ImageCore/Private/ImageCore.cpp
@@ -2293,7 +2293,15 @@ IMAGECORE_API void FImageCore::ResizeImageAllocDest(const FImageView & SourceIma
 
 IMAGECORE_API void FImageCore::ResizeImageAllocDest(const FImageView & SourceImage,FImage & DestImage,int32 DestSizeX, int32 DestSizeY, EResizeImageFilter Filter)
 {
-	ResizeImageAllocDest(SourceImage,DestImage,DestSizeX,DestSizeY,SourceImage.Format,SourceImage.GetGammaSpace(),Filter);
+	EGammaSpace DestGammaSpace = SourceImage.GetGammaSpace();
+	
+	// can't write Pow22 :
+	if ( DestGammaSpace == EGammaSpace::Pow22 )
+	{
+		DestGammaSpace = EGammaSpace::sRGB;
+	}
+
+	ResizeImageAllocDest(SourceImage,DestImage,DestSizeX,DestSizeY,SourceImage.Format,DestGammaSpace,Filter);
 }
 
 IMAGECORE_API void FImageCore::ResizeImageInPlace(FImage & Image,int32 DestSizeX, int32 DestSizeY, ERawImageFormat::Type DestFormat, EGammaSpace DestGammaSpace, EResizeImageFilter Filter)
