DDS: Adds an option to initialize the mips size information from header-only data, without allocating nor filling storage.

#rnx

diff --git a/Engine/Plugins/Interchange/Runtime/Source/Import/Private/Texture/InterchangeDDSTranslator.cpp b/Engine/Plugins/Interchange/Runtime/Source/Import/Private/Texture/InterchangeDDSTranslator.cpp
--- a/Engine/Plugins/Interchange/Runtime/Source/Import/Private/Texture/InterchangeDDSTranslator.cpp
+++ b/Engine/Plugins/Interchange/Runtime/Source/Import/Private/Texture/InterchangeDDSTranslator.cpp
@@ -80,7 +80,8 @@ namespace UE::Interchange::Private::InterchangeDDSTranslator
 		}
 
 		UE::DDS::EDDSError Error;
-		TUniquePtr<UE::DDS::FDDSFile> DDS = TUniquePtr<UE::DDS::FDDSFile>(UE::DDS::FDDSFile::CreateFromDDSInMemory(DDSSourceData.GetData(), DDSSourceData.Num(), &Error, bHeaderOnly));
+		const UE::DDS::EDDSReadMipMode MipReadMode = bHeaderOnly ? UE::DDS::EDDSReadMipMode::HeaderOnly : UE::DDS::EDDSReadMipMode::Full;
+		TUniquePtr<UE::DDS::FDDSFile> DDS = TUniquePtr<UE::DDS::FDDSFile>(UE::DDS::FDDSFile::CreateFromDDSInMemory(DDSSourceData.GetData(), DDSSourceData.Num(), &Error, MipReadMode));
 		if (!DDS.IsValid())
 		{
 			// NotADds is okay/expected , we try this on all image buffers
diff --git a/Engine/Source/Runtime/ImageCore/Private/DDSFile.cpp b/Engine/Source/Runtime/ImageCore/Private/DDSFile.cpp
--- a/Engine/Source/Runtime/ImageCore/Private/DDSFile.cpp
+++ b/Engine/Source/Runtime/ImageCore/Private/DDSFile.cpp
@@ -371,6 +371,25 @@ static void InitMip(FDDSMip* InMip, uint32 InWidth, uint32 InHeight, uint32 InDe
 	InMip->Data = 0;
 }
 
+static int64 InitMips(FDDSFile* InDDS, const FDXGIFormatInfo* InFormatInfo)
+{
+	int64 AllMipsSize = 0;
+	for (uint32 ArrayIndex = 0; ArrayIndex < InDDS->ArraySize; ++ArrayIndex)
+	{
+		for (uint32 MipIndex = 0; MipIndex < InDDS->MipCount; ++MipIndex)
+		{
+			FDDSMip* Mip = &InDDS->Mips[ArrayIndex * InDDS->MipCount + MipIndex];
+			uint32 MipWidth = MipDimension(InDDS->Width, MipIndex);
+			uint32 MipHeight = MipDimension(InDDS->Height, MipIndex);
+			uint32 MipDepth = MipDimension(InDDS->Depth, MipIndex);
+			InitMip(Mip, MipWidth, MipHeight, MipDepth, InFormatInfo);
+			AllMipsSize += Mip->DataSize;
+		}
+	}
+
+	return AllMipsSize;
+}
+
 EDDSError FDDSFile::Validate() const
 {
 	// Supported pixel format?
@@ -577,19 +596,7 @@ static EDDSError AllocateMips(FDDSFile* InDDS, const FDXGIFormatInfo* InFormatIn
 		// first pass, add up all sizes
 		//	then alloc it, second pass hand out all the pointers
 
-		int64 AllMipsSize = 0;
-		for (uint32 ArrayIndex = 0; ArrayIndex < InDDS->ArraySize; ++ArrayIndex) 
-		{
-			for (uint32 MipIndex = 0; MipIndex < InDDS->MipCount; ++MipIndex) 
-			{
-				FDDSMip* Mip = &InDDS->Mips[ArrayIndex * InDDS->MipCount + MipIndex];
-				uint32 MipWidth = MipDimension(InDDS->Width, MipIndex);
-				uint32 MipHeight = MipDimension(InDDS->Height, MipIndex);
-				uint32 MipDepth = MipDimension(InDDS->Depth, MipIndex);
-				InitMip(Mip, MipWidth, MipHeight, MipDepth, InFormatInfo);
-				AllMipsSize += Mip->DataSize;
-			}
-		}
+		int64 AllMipsSize = InitMips(InDDS, InFormatInfo);
 
 		InDDS->MipRawData.SetNumUninitialized(AllMipsSize);
 
@@ -748,8 +755,12 @@ int64 GetDDSHeaderMinimalSize()
 }
 
 
+FDDSFile* FDDSFile::CreateFromDDSInMemory(const uint8* InDDS, int64 InDDSSize, EDDSError* OutError, bool bHeaderOnly)
+{
+	return CreateFromDDSInMemory(InDDS, InDDSSize, OutError, static_cast<EDDSReadMipMode>(bHeaderOnly));
+}
 
-/* static */ FDDSFile* FDDSFile::CreateFromDDSInMemory(const uint8* InDDS, int64 InDDSSize, EDDSError *OutError, bool bHeaderOnly)
+/* static */ FDDSFile* FDDSFile::CreateFromDDSInMemory(const uint8* InDDS, int64 InDDSSize, EDDSError *OutError, EDDSReadMipMode InMipReadMode)
 {
 	// If no OutError passed in, redirect it to dummy storage on stack.
 	EDDSError DummyError;
@@ -804,9 +815,22 @@ int64 GetDDSHeaderMinimalSize()
 	FDDSFile* DDS = new FDDSFile();
 
 	*OutError = ParseHeader(DDS, &DDSHeader, &DDS10Header);
-	if (*OutError == EDDSError::OK && !bHeaderOnly)
+	if (*OutError == EDDSError::OK)
 	{
-		*OutError = ReadPayload(DDS, ReadCursor, ReadEnd);
+		if (InMipReadMode == UE::DDS::EDDSReadMipMode::Full)
+		{
+			*OutError = ReadPayload(DDS, ReadCursor, ReadEnd);
+		}
+		else if(InMipReadMode == UE::DDS::EDDSReadMipMode::HeaderWithMipInfo)
+		{
+			const FDXGIFormatInfo* FormatInfo = DXGIFormatGetInfo(DDS->DXGIFormat);
+
+			DDS->Mips.Empty();
+			DDS->Mips.SetNumZeroed(DDS->ArraySize * DDS->MipCount);
+
+			// Fill mips info without allocating storage
+			InitMips(DDS, FormatInfo);
+		}
 	}
 
 	if (*OutError != EDDSError::OK)
diff --git a/Engine/Source/Runtime/ImageCore/Public/DDSFile.h b/Engine/Source/Runtime/ImageCore/Public/DDSFile.h
--- a/Engine/Source/Runtime/ImageCore/Public/DDSFile.h
+++ b/Engine/Source/Runtime/ImageCore/Public/DDSFile.h
@@ -161,6 +161,14 @@ namespace UE { namespace DDS
 		D3D10,	// Always write a D3D10-format .DDS file
 	};
 
+	// Mip data read modes
+	enum class EDDSReadMipMode
+	{
+		Full,
+		HeaderOnly,
+		HeaderWithMipInfo,
+	};
+
 	// Error codes
 	enum class EDDSError
 	{
@@ -286,11 +294,14 @@ namespace UE { namespace DDS
 		IMAGECORE_API bool IsValidTextureArray() const;
 		IMAGECORE_API bool IsValidTextureVolume() const;
 
+		UE_DEPRECATED(5.6, "This function is now deprecated, please use CreateFromDDSInMemory with EDDSReadMipMode instead.")
+		IMAGECORE_API static FDDSFile* CreateFromDDSInMemory(const uint8* InDDS, int64 InDDSSize, EDDSError* OutError, bool bHeaderOnly);
+
 		// Used for loading from a DDS file image in memory.
 		// 
 		// On error, returns nullptr. If a non-null OutError is supplied, error information
 		// is written there.
-		IMAGECORE_API static FDDSFile* CreateFromDDSInMemory(const uint8* InDDS, int64 InDDSSize, EDDSError* OutError=nullptr, bool bHeaderOnly = false);
+		IMAGECORE_API static FDDSFile* CreateFromDDSInMemory(const uint8* InDDS, int64 InDDSSize, EDDSError* OutError=nullptr, EDDSReadMipMode InMipReadMode = EDDSReadMipMode::Full);
 		
 		// is this buffer a DDS ?
 		IMAGECORE_API static bool IsADDS(const uint8* InDDS, int64 InDDSSize);
