FImage move constructor

diff --git a/Engine/Source/Runtime/ImageCore/Private/ImageCore.cpp b/Engine/Source/Runtime/ImageCore/Private/ImageCore.cpp
--- a/Engine/Source/Runtime/ImageCore/Private/ImageCore.cpp
+++ b/Engine/Source/Runtime/ImageCore/Private/ImageCore.cpp
@@ -998,18 +998,6 @@ void FImage::FreeData(bool bAsyncDetached)
 	}
 }
 
-void FImage::Swap(FImage & Other)
-{
-	::Swap(RawData,Other.RawData);
-	// FImageInfo can be moved with assignment
-	FImageInfo * ThisImage = this;
-	FImageInfo * OtherImage = &Other;
-	FImageInfo Temp;
-	Temp = *ThisImage;
-	*ThisImage = *OtherImage;
-	*OtherImage = Temp;
-}
-
 void FImage::ChangeFormat(ERawImageFormat::Type DestFormat, EGammaSpace DestGammaSpace)
 {
 	if ( Format == DestFormat &&
diff --git a/Engine/Source/Runtime/ImageCore/Public/ImageCore.h b/Engine/Source/Runtime/ImageCore/Public/ImageCore.h
--- a/Engine/Source/Runtime/ImageCore/Public/ImageCore.h
+++ b/Engine/Source/Runtime/ImageCore/Public/ImageCore.h
@@ -422,6 +422,26 @@ public:
 	FImage()
 	{
 	}
+	
+	FImage(const FImage & CopyFrom) = default;
+	FImage& operator=(const FImage & CopyFrom) = default;
+
+	// move constructor
+	FORCEINLINE FImage(FImage && MoveFrom) : FImageInfo(MoveFrom), RawData(MoveTemp(MoveFrom.RawData))
+	{
+		MoveFrom.Reset();
+	}
+	
+	// move assignment
+	FImage& operator=(FImage && MoveFrom)
+	{
+		if (this != &MoveFrom)
+		{
+			Swap(MoveFrom);
+			MoveFrom.Reset();
+		}
+		return *this;
+	}
 
 	/**
 	 * Creates and initializes a new image with the specified number of slices.
@@ -464,11 +484,24 @@ public:
 
 	// Free FImage.RawData
 	//	if bAsyncDetached free is done on a Task, not immediately, but the RawData member is empty upon return
+	//	note that FreeData does not change the FImageInfo fields, only frees RawData, follow up with Reset() if desired
 	IMAGECORE_API void FreeData(bool bAsyncDetached);
 
 	// Swap the contents of this FImage with another
-	IMAGECORE_API void Swap(FImage & Other);
+	FORCEINLINE void Swap(FImage & Other)
+	{
+		::Swap(RawData,Other.RawData);
+		// FImageInfo part be swapped with assignment
+		::Swap<FImageInfo>(*this,Other); // Templates/UnrealTemplate.h
+	}
 	
+	// Reset this FImage to empty default constructed state
+	FORCEINLINE void Reset()
+	{
+		FImage Nada;
+		Swap(Nada);
+	}
+
 	// implicit conversion to FImageView
 	// allows FImage to be passed to functions that take FImageView
 	// functions that read images should use "const FImageView &" as their argument type
