FImageCore::ComputeChannelLinearMinMax disable test

diff --git a/Engine/Source/Runtime/ImageCore/Private/ImageParallelFor.cpp b/Engine/Source/Runtime/ImageCore/Private/ImageParallelFor.cpp
--- a/Engine/Source/Runtime/ImageCore/Private/ImageParallelFor.cpp
+++ b/Engine/Source/Runtime/ImageCore/Private/ImageParallelFor.cpp
@@ -531,16 +531,6 @@ void FImageCore::ComputeChannelLinearMinMax(const FImageView & Image, FLinearCol
 	check( AccumView.GetImageSizeBytes() == 32 );
 
 	ComputeChannelLinearMinMax_Generic(AccumView,OutMin,OutMax);
-
-	#if 1
-	// compare ComputeChannelLinearMinMax_Generic for sanity check
-	{
-		FLinearColor RefMin,RefMax;
-		ComputeChannelLinearMinMax_Generic(Image,RefMin,RefMax);
-		check( OutMin == RefMin );
-		check( OutMax == RefMax );
-	}
-	#endif
 }
 
 bool FImageCore::ScaleChannelsSoMinMaxIsInZeroToOne(const FImageView & Image)
