Replace some usages of FORCEINLINE with inline in Oodle modules.

For context about this change, see this doc: https://docs.google.com/document/d/1LjpT7NpzTecfXdFNyv5KM6sCu67827BRljNi1iaWVSE/edit?tab=t.0#heading=h.fvtz8gjukz7x

Includes modules: OodleNetworkHandlerComponent, ImageCore

Reviewers were automatically determined. Please comment if anyone else should be added.

#rnx

diff --git a/Engine/Plugins/Compression/OodleNetwork/Source/Public/OodleNetworkArchives.h b/Engine/Plugins/Compression/OodleNetwork/Source/Public/OodleNetworkArchives.h
--- a/Engine/Plugins/Compression/OodleNetwork/Source/Public/OodleNetworkArchives.h
+++ b/Engine/Plugins/Compression/OodleNetwork/Source/Public/OodleNetworkArchives.h
@@ -76,7 +76,7 @@ protected:
 		 *
 		 * @return	Returns the property value
 		 */
-		FORCEINLINE Type Get()
+		inline Type Get()
 		{
 			return Value;
 		}
@@ -87,7 +87,7 @@ protected:
 		 * @param Ar		The archive being serialized to
 		 * @param InValue	The new value
 		 */
-		FORCEINLINE void Set(FOodleNetworkArchiveBase& Ar, Type InValue)
+		inline void Set(FOodleNetworkArchiveBase& Ar, Type InValue)
 		{
 			Value = InValue;
 
@@ -104,7 +104,7 @@ protected:
 		 *
 		 * @param Ar	The archive being serialized to/from
 		 */
-		FORCEINLINE void Serialize(FArchive& Ar)
+		inline void Serialize(FArchive& Ar)
 		{
 			if (Offset == 0)
 			{
@@ -235,7 +235,7 @@ public:
 	 * @param Value		The struct value to serialize
 	 * @return			Returns a chained-reference to the archive
 	 */
-	FORCEINLINE friend FArchive& operator<<(FArchive& Ar, FOodleCompressedData& Value)
+	inline friend FArchive& operator<<(FArchive& Ar, FOodleCompressedData& Value)
 	{
 		Value.Serialize(Ar);
 		return Ar;
@@ -248,7 +248,7 @@ public:
 	 * @param Value		The rewritable value to serialize
 	 * @return			Returns a chained-reference to the archive
 	 */
-	template<class Type> FORCEINLINE friend FArchive& operator<<(FArchive& Ar, TRewritable<Type>& Value)
+	template<class Type> inline friend FArchive& operator<<(FArchive& Ar, TRewritable<Type>& Value)
 	{
 		Value.Serialize(Ar);
 		return Ar;
diff --git a/Engine/Plugins/Compression/OodleNetwork/Source/Public/OodleNetworkHandlerComponent.h b/Engine/Plugins/Compression/OodleNetwork/Source/Public/OodleNetworkHandlerComponent.h
--- a/Engine/Plugins/Compression/OodleNetwork/Source/Public/OodleNetworkHandlerComponent.h
+++ b/Engine/Plugins/Compression/OodleNetwork/Source/Public/OodleNetworkHandlerComponent.h
@@ -152,7 +152,7 @@ public:
 	 * @param CompressedLength		The compressed size of the input packet
 	 * @param DecompressedLength	The decompressed size of the input packet
 	 */
-	FORCEINLINE void IncomingStats(uint32 CompressedLength, uint32 DecompressedLength)
+	inline void IncomingStats(uint32 CompressedLength, uint32 DecompressedLength)
 	{
 		InCompressedLength += CompressedLength;
 		TotalInCompressedLength += CompressedLength;
@@ -168,7 +168,7 @@ public:
 	 * @param CompressedLength		The compressed size of the output packet
 	 * @param UncompressedLength	The uncompressed size of the output packets
 	 */
-	FORCEINLINE void OutgoingStats(uint32 CompressedLength, uint32 UncompressedLength)
+	inline void OutgoingStats(uint32 CompressedLength, uint32 UncompressedLength)
 	{
 		OutCompressedLength += CompressedLength;
 		TotalOutCompressedLength += CompressedLength;
@@ -181,7 +181,7 @@ public:
 	/**
 	 * Checks to see if the main stats are due an update, and triggers an update if so
 	 */
-	FORCEINLINE void CheckForUpdate()
+	inline void CheckForUpdate()
 	{
 		float DeltaTime = static_cast<float>(FPlatformTime::Seconds() - LastStatsUpdate);
 
diff --git a/Engine/Source/Runtime/ImageCore/Public/ImageCore.h b/Engine/Source/Runtime/ImageCore/Public/ImageCore.h
--- a/Engine/Source/Runtime/ImageCore/Public/ImageCore.h
+++ b/Engine/Source/Runtime/ImageCore/Public/ImageCore.h
@@ -87,14 +87,14 @@ namespace ERawImageFormat
 	IMAGECORE_API const FLinearColor GetOnePixelLinear(const void * PixelData,Type Format,EGammaSpace Gamma);
 	
 	// Get one pixel of Format type from PixelData and return in Linear color
-	FORCEINLINE const FLinearColor GetOnePixelLinear(const void * PixelData,Type Format,bool bSRGB)
+	inline const FLinearColor GetOnePixelLinear(const void * PixelData,Type Format,bool bSRGB)
 	{
 		return GetOnePixelLinear(PixelData,Format,bSRGB ? EGammaSpace::sRGB : EGammaSpace::Linear);
 	}
 
 	// G8 and BGRA8 are affected by Gamma, 16/32 is NOT (they are always Linear)
 	// this happens to be the same as IsU8Channels() at the moment, but don't assume that
-	FORCEINLINE bool GetFormatNeedsGammaSpace(Type Format)
+	inline bool GetFormatNeedsGammaSpace(Type Format)
 	{
 		if ( Format == G8 || Format == BGRA8 )
 		{
@@ -110,7 +110,7 @@ namespace ERawImageFormat
 	// when converting from pixel bags that don't store gamma to FImage
 	//	you can use this to guess the gamma they probably wanted
 	//	do not use when a real gamma flag is available!
-	FORCEINLINE EGammaSpace GetDefaultGammaSpace(Type Format)
+	inline EGammaSpace GetDefaultGammaSpace(Type Format)
 	{
 		if ( GetFormatNeedsGammaSpace(Format) )
 		{
@@ -163,7 +163,7 @@ struct FImageInfo
 	{
 	}
 	
-	FORCEINLINE bool operator == (const FImageInfo & rhs) const
+	inline bool operator == (const FImageInfo & rhs) const
 	{
 		return 
 			SizeX == rhs.SizeX &&
@@ -173,7 +173,7 @@ struct FImageInfo
 			GammaSpace == rhs.GammaSpace;
 	}
 
-	FORCEINLINE bool IsImageInfoValid() const
+	inline bool IsImageInfoValid() const
 	{
 		if ( SizeX < 0 || SizeY < 0 || NumSlices < 0 ) return false;
 		if ( Format == ERawImageFormat::Invalid ) return false;
@@ -182,7 +182,7 @@ struct FImageInfo
 		return true;
 	}
 
-	FORCEINLINE bool IsGammaCorrected() const
+	inline bool IsGammaCorrected() const
 	{
 		return GammaSpace != EGammaSpace::Linear;
 	}
@@ -192,37 +192,37 @@ struct FImageInfo
 	 *
 	 * @return Bytes per pixel.
 	 */
-	FORCEINLINE int64 GetBytesPerPixel() const
+	inline int64 GetBytesPerPixel() const
 	{
 		return ERawImageFormat::GetBytesPerPixel(Format);
 	}
 	
-	FORCEINLINE int64 GetNumPixels() const
+	inline int64 GetNumPixels() const
 	{
 		return (int64) SizeX * SizeY * NumSlices;
 	}
 	
-	FORCEINLINE int64 GetImageSizeBytes() const
+	inline int64 GetImageSizeBytes() const
 	{
 		return GetNumPixels() * GetBytesPerPixel();
 	}
 	
-	FORCEINLINE int64 GetSliceNumPixels() const
+	inline int64 GetSliceNumPixels() const
 	{
 		return (int64) SizeX * SizeY;
 	}
 	
-	FORCEINLINE int64 GetSliceSizeBytes() const
+	inline int64 GetSliceSizeBytes() const
 	{
 		return GetSliceNumPixels() * GetBytesPerPixel();
 	}
 	
-	FORCEINLINE int64 GetWidth()  const { return SizeX; }
-	FORCEINLINE int64 GetHeight() const { return SizeY; }
+	inline int64 GetWidth()  const { return SizeX; }
+	inline int64 GetHeight() const { return SizeY; }
 	
-	FORCEINLINE int64 GetStrideBytes()  const { return SizeX * GetBytesPerPixel(); }
+	inline int64 GetStrideBytes()  const { return SizeX * GetBytesPerPixel(); }
 
-	FORCEINLINE EGammaSpace GetGammaSpace() const
+	inline EGammaSpace GetGammaSpace() const
 	{
 		// Gamma is ignored unless GetFormatNeedsGammaSpace, so make sure it is Linear
 		check( GetFormatNeedsGammaSpace(Format) || GammaSpace == EGammaSpace::Linear );
@@ -231,7 +231,7 @@ struct FImageInfo
 	}
 	
 	// get offset of a pixel from the base pointer, in bytes
-	FORCEINLINE int64 GetPixelOffsetBytes(int32 X,int32 Y,int32 Slice = 0) const
+	inline int64 GetPixelOffsetBytes(int32 X,int32 Y,int32 Slice = 0) const
 	{
 		checkSlow( X >= 0 && X < SizeX );
 		checkSlow( Y >= 0 && Y < SizeY );
@@ -334,7 +334,7 @@ struct FImageView : public FImageInfo
 	}
 	
 	// get a pointer to a pixel
-	FORCEINLINE void * GetPixelPointer(int32 X,int32 Y,int32 Slice=0) const
+	inline void * GetPixelPointer(int32 X,int32 Y,int32 Slice=0) const
 	{
 		uint8 * Ptr = (uint8 *)RawData;
 		Ptr += GetPixelOffsetBytes(X,Y,Slice);
@@ -430,7 +430,7 @@ public:
 	FImage& operator=(const FImage & CopyFrom) = default;
 
 	// move constructor
-	FORCEINLINE FImage(FImage && MoveFrom) : FImageInfo(MoveFrom), RawData(MoveTemp(MoveFrom.RawData))
+	inline FImage(FImage && MoveFrom) : FImageInfo(MoveFrom), RawData(MoveTemp(MoveFrom.RawData))
 	{
 		MoveFrom.Reset();
 	}
@@ -458,7 +458,7 @@ public:
 	IMAGECORE_API FImage(int32 InSizeX, int32 InSizeY, int32 InNumSlices, ERawImageFormat::Type InFormat, EGammaSpace InGammaSpace);
 
 	// note: changed default GammaSpace from Linear to GetDefaultGammaSpace
-	FORCEINLINE FImage(int32 InSizeX, int32 InSizeY, int32 InNumSlices, ERawImageFormat::Type InFormat) :
+	inline FImage(int32 InSizeX, int32 InSizeY, int32 InNumSlices, ERawImageFormat::Type InFormat) :
 		FImage(InSizeX,InSizeY,InNumSlices,InFormat,ERawImageFormat::GetDefaultGammaSpace(InFormat))
 	{
 	}
@@ -471,13 +471,13 @@ public:
 	 * @param InFormat - The image format.
 	 * @param InGammaSpace - (optional) Image Gamma Space.  GetDefaultGammaSpace() if not provided.
 	 */
-	FORCEINLINE FImage(int32 InSizeX, int32 InSizeY, ERawImageFormat::Type InFormat, EGammaSpace InGammaSpace) :
+	inline FImage(int32 InSizeX, int32 InSizeY, ERawImageFormat::Type InFormat, EGammaSpace InGammaSpace) :
 		FImage(InSizeX,InSizeY,1,InFormat,InGammaSpace)
 	{
 	}
 
 	// note: changed default GammaSpace from Linear to GetDefaultGammaSpace
-	FORCEINLINE FImage(int32 InSizeX, int32 InSizeY, ERawImageFormat::Type InFormat) :
+	inline FImage(int32 InSizeX, int32 InSizeY, ERawImageFormat::Type InFormat) :
 		FImage(InSizeX,InSizeY,1,InFormat,ERawImageFormat::GetDefaultGammaSpace(InFormat))
 	{
 	}
@@ -491,7 +491,7 @@ public:
 	IMAGECORE_API void FreeData(bool bAsyncDetached);
 
 	// Swap the contents of this FImage with another
-	FORCEINLINE void Swap(FImage & Other)
+	inline void Swap(FImage & Other)
 	{
 		::Swap(RawData,Other.RawData);
 		// FImageInfo part be swapped with assignment
@@ -499,7 +499,7 @@ public:
 	}
 	
 	// Reset this FImage to empty default constructed state
-	FORCEINLINE void Reset()
+	inline void Reset()
 	{
 		FImage Nada;
 		Swap(Nada);
@@ -508,7 +508,7 @@ public:
 	// implicit conversion to FImageView
 	// allows FImage to be passed to functions that take FImageView
 	// functions that read images should use "const FImageView &" as their argument type
-	FORCEINLINE operator FImageView () const
+	inline operator FImageView () const
 	{
 		// copy the shared FImageInfo part :
 		const FImageInfo & Info = *this;
@@ -562,7 +562,7 @@ public:
 	 */
 	IMAGECORE_API void Linearize(uint8 SourceEncoding, FImage& DestImage) const;
 
-	FORCEINLINE void Linearize(FImage& DestImage) const
+	inline void Linearize(FImage& DestImage) const
 	{
 		Linearize(0,DestImage);
 	}
@@ -582,7 +582,7 @@ public:
 	 */
 	IMAGECORE_API void Init(int32 InSizeX, int32 InSizeY, int32 InNumSlices, ERawImageFormat::Type InFormat, EGammaSpace InGammaSpace);
 	
-	FORCEINLINE void Init(int32 InSizeX, int32 InSizeY, int32 InNumSlices, ERawImageFormat::Type InFormat)
+	inline void Init(int32 InSizeX, int32 InSizeY, int32 InNumSlices, ERawImageFormat::Type InFormat)
 	{
 		Init(InSizeX,InSizeY,InNumSlices,InFormat,ERawImageFormat::GetDefaultGammaSpace(InFormat));
 	}
@@ -598,7 +598,7 @@ public:
 	 */
 	IMAGECORE_API void Init(int32 InSizeX, int32 InSizeY, ERawImageFormat::Type InFormat, EGammaSpace InGammaSpace);
 	
-	FORCEINLINE void Init(int32 InSizeX, int32 InSizeY, ERawImageFormat::Type InFormat)
+	inline void Init(int32 InSizeX, int32 InSizeY, ERawImageFormat::Type InFormat)
 	{
 		Init(InSizeX,InSizeY,InFormat,ERawImageFormat::GetDefaultGammaSpace(InFormat));
 	}
@@ -612,7 +612,7 @@ public:
 	IMAGECORE_API void Init(const FImageInfo & Info);
 	
 	// get a pointer to a pixel
-	FORCEINLINE void * GetPixelPointer(int32 X,int32 Y,int32 Slice=0) const
+	inline void * GetPixelPointer(int32 X,int32 Y,int32 Slice=0) const
 	{
 		int64 Offset = GetPixelOffsetBytes(X,Y,Slice);
 		return (void *)&RawData[Offset];
@@ -818,7 +818,7 @@ public:
 	IMAGECORE_API void ChangeFormat(ERawImageFormat::Type DestFormat, EGammaSpace DestGammaSpace);
 
 
-	FORCEINLINE bool IsValid() const
+	inline bool IsValid() const
 	{
 		bool bMipInfoValid = true;
 		for (const FMipInfo& MipInfo : SubImages)
@@ -833,7 +833,7 @@ public:
 		return true;
 	}
 
-	FORCEINLINE bool GetMipDimensions(int32 MipLevel, int32& OutWidth, int32& OutHeight) const
+	inline bool GetMipDimensions(int32 MipLevel, int32& OutWidth, int32& OutHeight) const
 	{
 		if (!SubImages.IsValidIndex(MipLevel))
 		{
@@ -847,7 +847,7 @@ public:
 		return true;
 	}
 
-	FORCEINLINE int32 GetMipCount() const 
+	inline int32 GetMipCount() const 
 	{
 		return SubImages.Num();
 	}
@@ -871,7 +871,7 @@ public:
 		RawData.Append(MoveTemp(Buffer));
 	}
 
-	FORCEINLINE int64 GetNumPixels()
+	inline int64 GetNumPixels()
 	{
 		return (RawData.Num() / ERawImageFormat::GetBytesPerPixel(Format));
 	}
diff --git a/Engine/Source/Runtime/ImageCore/Public/ImageCoreDelta.h b/Engine/Source/Runtime/ImageCore/Public/ImageCoreDelta.h
--- a/Engine/Source/Runtime/ImageCore/Public/ImageCoreDelta.h
+++ b/Engine/Source/Runtime/ImageCore/Public/ImageCoreDelta.h
@@ -38,14 +38,14 @@ struct FImageViewStrided : public FImageInfo
 	}
 
 	// replaces the call in FImageInfo but is NOT virtual ; beware
-	FORCEINLINE int64 GetStrideBytes()  const { return StrideBytes; }
+	inline int64 GetStrideBytes()  const { return StrideBytes; }
 	
 	// if IsStrideWidth, then you can convert back to FImageView
-	FORCEINLINE bool IsStrideWidth() const { return StrideBytes == SizeX * GetBytesPerPixel(); }
+	inline bool IsStrideWidth() const { return StrideBytes == SizeX * GetBytesPerPixel(); }
 
 	// get offset of a pixel from the base pointer, in bytes
 	// replaces the call in FImageInfo but is NOT virtual ; beware
-	FORCEINLINE int64 GetPixelOffsetBytes(int32 X,int32 Y,int32 Slice = 0) const
+	inline int64 GetPixelOffsetBytes(int32 X,int32 Y,int32 Slice = 0) const
 	{
 		checkSlow( X >= 0 && X < SizeX );
 		checkSlow( Y >= 0 && Y < SizeY );
@@ -60,7 +60,7 @@ struct FImageViewStrided : public FImageInfo
 
 	// queries like GetImageSizeBytes are ambiguous ; do you mean the used pixels or the stride?
 
-	FORCEINLINE uint8 * GetRowPointer(int32 Y,int32 Slice = 0) const
+	inline uint8 * GetRowPointer(int32 Y,int32 Slice = 0) const
 	{
 		checkSlow( Y >= 0 && Y < SizeY );
 		checkSlow( Slice >= 0 && Slice < NumSlices );
@@ -72,14 +72,14 @@ struct FImageViewStrided : public FImageInfo
 	}
 	
 	// get a pointer to a pixel
-	FORCEINLINE void * GetPixelPointer(int32 X,int32 Y,int32 Slice=0) const
+	inline void * GetPixelPointer(int32 X,int32 Y,int32 Slice=0) const
 	{
 		uint8 * Ptr = (uint8 *)RawData;
 		Ptr += GetPixelOffsetBytes(X,Y,Slice);
 		return (void *)Ptr;
 	}
 	
-	FORCEINLINE const FImageViewStrided GetPortion(int64 PortionStartX,int64 PortionSizeX, int64 PortionStartY,int64 PortionSizeY) const
+	inline const FImageViewStrided GetPortion(int64 PortionStartX,int64 PortionSizeX, int64 PortionStartY,int64 PortionSizeY) const
 	{
 		check( PortionStartX >= 0 && (PortionStartX + PortionSizeX) <= SizeX );
 		check( PortionStartY >= 0 && (PortionStartY + PortionSizeY) <= SizeY );
